#UE4
- Removed unused code, improved parsing stats metadata by the profiler client
- Fixed a problem with the incomplete metadata when using the live preview in the profiler
- Misc fixes in the stats metadata implementation in the profilerservice/profilerclient, mostly related to stats2 migration some time ago
- FIxes related to reading old captured stats files

diff --git a/Engine/Source/Developer/ProfilerClient/Private/ProfilerClientManager.cpp b/Engine/Source/Developer/ProfilerClient/Private/ProfilerClientManager.cpp
--- a/Engine/Source/Developer/ProfilerClient/Private/ProfilerClientManager.cpp
+++ b/Engine/Source/Developer/ProfilerClient/Private/ProfilerClientManager.cpp
@@ -424,6 +424,7 @@ void FServiceConnection::Initialize( const FProfilerServiceAuthorize2& Message,
 void FProfilerClientManager::HandleServiceMetaDataMessage( const FProfilerServiceMetaData& Message, const IMessageContextRef& Context )
 {
 #if STATS
+	// @TODO yrx 2014-04-14 Not used.
 	if (ActiveSessionId.IsValid() && Connections.Find(Message.InstanceId) != nullptr)
 	{
 		FServiceConnection& Connection = *Connections.Find(Message.InstanceId);
@@ -797,6 +798,8 @@ void FServiceConnection::UpdateMetaData()
 	for (auto It = CurrentThreadState.ShortNameToLongName.CreateConstIterator(); It; ++It)
 	{
 		FStatMessage const& LongName = It.Value();
+		const FName GroupName = LongName.NameAndInfo.GetGroupName();
+
 		uint32 StatType = STATTYPE_Error;
 		if (LongName.NameAndInfo.GetField<EStatDataType>() == EStatDataType::ST_int64)
 		{
@@ -821,6 +824,13 @@ void FServiceConnection::UpdateMetaData()
 		{
 			FindOrAddStat(LongName.NameAndInfo, StatType);
 		}
+
+		// Threads metadata.
+		const bool bIsThread = FStatConstants::NAME_ThreadGroup == GroupName;
+		if( bIsThread )
+		{
+			FindOrAddThread( LongName.NameAndInfo );
+		}	
 	}
 }
 
@@ -867,7 +877,7 @@ int32 FServiceConnection::FindOrAddStat( const FStatNameAndInfo& StatNameAndInfo
 		}
 		StatDescription.StatType = StatType;
 
-		if( GroupName == NAME_None )
+		if( GroupName == NAME_None && Stream.Header.Version == EStatMagicNoHeader::NO_VERSION )
 		{	
 			// @todo Add more ways to group the stats.
 			const int32 Thread_Pos = StatDescription.Name.Find( TEXT("Thread_") );
@@ -919,23 +929,22 @@ int32 FServiceConnection::FindOrAddThread(const FStatNameAndInfo& Thread)
 {
 	SCOPE_CYCLE_COUNTER(STAT_PC_FindOrAddThread);
 
-	const FName LongName = Thread.GetRawName();
-	int32* const ThreadIDPtr = ThreadNameArray.Find( LongName );
-	int32 ThreadID = ThreadIDPtr != nullptr ? *ThreadIDPtr : -1;
-	if (!ThreadIDPtr)
-	{
-		// meta data has been updated
-		CurrentData.MetaDataUpdated = true;
+	CurrentThreadState.Threads;
 
-		// get the thread description
-		FString Desc;
-		Thread.GetDescription(Desc);
-		ThreadID = FStatsUtils::ParseThreadID( FName( *Desc ) );
-		ThreadNameArray.Add(LongName, ThreadID);
+	const FName ShortName = Thread.GetShortName();
+	FString Desc;
+	Thread.GetDescription( Desc );
+	const uint32 ThreadID = FStatsUtils::ParseThreadID( Desc );
 
+	{
 		// add to the meta data
-		FScopeLock ScopeLock(&CriticalSection);
-		MetaData.ThreadDescriptions.Add(ThreadID, Thread.GetShortName().ToString());
+		FScopeLock ScopeLock( &CriticalSection );
+		const int32 OldNum = MetaData.ThreadDescriptions.Num();
+		MetaData.ThreadDescriptions.Add( ThreadID, ShortName.ToString() );
+		const int32 NewNum = MetaData.ThreadDescriptions.Num();
+
+		// meta data has been updated
+		CurrentData.MetaDataUpdated = CurrentData.MetaDataUpdated || OldNum != NewNum;
 	}
 
 	return ThreadID;
diff --git a/Engine/Source/Developer/ProfilerClient/Private/ProfilerClientManager.h b/Engine/Source/Developer/ProfilerClient/Private/ProfilerClientManager.h
--- a/Engine/Source/Developer/ProfilerClient/Private/ProfilerClientManager.h
+++ b/Engine/Source/Developer/ProfilerClient/Private/ProfilerClientManager.h
@@ -52,9 +52,6 @@ struct FServiceConnection
 	/** Provides an FName to GroupId mapping */
 	TMap<FName, int32> GroupNameArray;
 
-	/** Provides an FName to ThreadId mapping */
-	TMap<FName, int32> ThreadNameArray;
-
 	/** Provides the long stat name to StatId mapping. */
 	TMap<FName,int32> LongNameToStatID;
 
diff --git a/Engine/Source/Developer/ProfilerService/Private/ProfilerServiceManager.cpp b/Engine/Source/Developer/ProfilerService/Private/ProfilerServiceManager.cpp
--- a/Engine/Source/Developer/ProfilerService/Private/ProfilerServiceManager.cpp
+++ b/Engine/Source/Developer/ProfilerService/Private/ProfilerServiceManager.cpp
@@ -315,10 +315,6 @@ FProfilerServiceManager::FProfilerServiceManager()
 	NewMetaData.SecondsPerCycle = FPlatformTime::GetSecondsPerCycle();
 	PingDelegate = FTickerDelegate::CreateRaw(this, &FProfilerServiceManager::HandlePing);
 	DataFrame.Frame = 0;
-
-	MetaData.ThreadDescriptions.FindOrAdd(GGameThreadId) = FName(NAME_GameThread).GetPlainNameString();
-
-	UpdateMetaData();
 }
 
 
@@ -388,23 +384,9 @@ void FProfilerServiceManager::StopCapture()
 #endif
 }
 
-
-void FProfilerServiceManager::AddStatDescription(FStatDescription& Data)
-{
-	MetaData.StatDescriptions.FindOrAdd(Data.ID) = Data;
-	NewMetaData.StatDescriptions.FindOrAdd(Data.ID) = Data;
-}
-
-
-void FProfilerServiceManager::AddGroupDescription(FStatGroupDescription& Data)
-{
-	MetaData.GroupDescriptions.FindOrAdd(Data.ID) = Data;
-	NewMetaData.GroupDescriptions.FindOrAdd(Data.ID) = Data;
-}
-
-
 void FProfilerServiceManager::UpdateMetaData()
 {
+	// @TODO yrx 2014-04-13 Obsolete, remove later.
 	// update the thread descriptions only if there has been a change
 	int32 OldCount = FRunnableThread::GetThreadRegistry().GetThreadCount();
 	if (FRunnableThread::GetThreadRegistry().IsUpdated())
@@ -447,9 +429,6 @@ void FProfilerServiceManager::SendMetaData(const FMessageAddress& client)
 
 void FProfilerServiceManager::StartFrame(uint32 FrameNumber, double FrameStart)
 {
-	// update the threads
-	UpdateMetaData();
-
 	// send data to the clients
 	if (DataFrame.Frame > 0)
 	{
@@ -541,6 +520,7 @@ IProfilerServiceManagerPtr FProfilerServiceManager::CreateSharedServiceManager()
 void FProfilerServiceManager::SetPreviewState( const FMessageAddress& ClientAddress, const bool bRequestedPreviewState )
 {
 #if STATS
+	const FStatsThreadState& Stats = FStatsThreadState::GetLocalState();
 	FClientData* Client = ClientData.Find( ClientAddress );
 	if( Client )
 	{
@@ -553,15 +533,15 @@ void FProfilerServiceManager::SetPreviewState( const FMessageAddress& ClientAddr
 				// enable stat capture
 				if (PreviewClients.Num() == 0)
 				{
-					FStatsThreadState::GetLocalState().NewFrameDelegate.AddRaw(this, &FProfilerServiceManager::HandleNewFrame);
+					Stats.NewFrameDelegate.AddRaw( this, &FProfilerServiceManager::HandleNewFrame );
 					StatsMasterEnableAdd();
 				}
 				PreviewClients.Add(ClientAddress);
 				Client->Preview = true;
 				if (MessageEndpoint.IsValid())
 				{
-					Client->CurrentFrame = FStatsThreadState::GetLocalState().CurrentGameFrame;
-					MessageEndpoint->Send(new FProfilerServicePreviewAck(InstanceId, FStatsThreadState::GetLocalState().CurrentGameFrame), ClientAddress);
+					Client->CurrentFrame = Stats.CurrentGameFrame;
+					MessageEndpoint->Send( new FProfilerServicePreviewAck( InstanceId, Stats.CurrentGameFrame ), ClientAddress );
 				}
 				SendMetaData(ClientAddress);
 			}
@@ -574,7 +554,7 @@ void FProfilerServiceManager::SetPreviewState( const FMessageAddress& ClientAddr
 				if (PreviewClients.Num() == 0)
 				{
 					// disable stat capture
-					FStatsThreadState::GetLocalState().NewFrameDelegate.RemoveRaw(this, &FProfilerServiceManager::HandleNewFrame);
+					Stats.NewFrameDelegate.RemoveRaw( this, &FProfilerServiceManager::HandleNewFrame );
 					StatsMasterEnableAdd();
 				}
 			}
@@ -705,7 +685,8 @@ void FProfilerServiceManager::HandleServiceFileChunkMessage( const FProfilerServ
 void FProfilerServiceManager::HandleServiceSubscribeMessage( const FProfilerServiceSubscribe& Message, const IMessageContextRef& Context )
 {
 #if STATS
-	if (Message.SessionId == SessionId && Message.InstanceId == InstanceId && !ClientData.Contains(Context->GetSender()))
+	const FMessageAddress& MsgAddress = Context->GetSender();
+	if( Message.SessionId == SessionId && Message.InstanceId == InstanceId && !ClientData.Contains( MsgAddress ) )
 	{
 		UE_LOG(LogProfile, Log, TEXT("Added a client" ));
 
@@ -715,10 +696,13 @@ void FProfilerServiceManager::HandleServiceSubscribeMessage( const FProfilerServ
 		Data.StatsWriteFile.WriteHeader();
 
 		// add to the client list
-		ClientData.Add(Context->GetSender(), Data);
+		ClientData.Add( MsgAddress, Data );
 		// send authorized and stat descriptions
-		const TArray<uint8>& OutData = ClientData.Find(Context->GetSender())->StatsWriteFile.GetOutData();
-		MessageEndpoint->Send(new FProfilerServiceAuthorize2(SessionId, InstanceId, OutData), Context->GetSender());
+		const TArray<uint8>& OutData = ClientData.Find( MsgAddress )->StatsWriteFile.GetOutData();
+		MessageEndpoint->Send( new FProfilerServiceAuthorize2( SessionId, InstanceId, OutData ), MsgAddress );
+
+		const FStatsThreadState& Stats = FStatsThreadState::GetLocalState();
+		ClientData.Find( MsgAddress )->MetadataSize = Stats.ShortNameToLongName.Num();
 
 		// initiate the ping callback
 		if (ClientData.Num() == 1)
@@ -761,15 +745,28 @@ void FProfilerServiceManager::HandleNewFrame(int64 Frame)
 	// package it up and send to the clients
 	if (MessageEndpoint.IsValid())
 	{
+		const FStatsThreadState& Stats = FStatsThreadState::GetLocalState();
+		const int32 CurrentMetadataSize = Stats.ShortNameToLongName.Num();
+
 		// update preview clients with the current data
 		for (auto It = PreviewClients.CreateConstIterator(); It; ++It)
 		{
 			FClientData& Client = *ClientData.Find(*It);
+
 			while(Client.CurrentFrame < Frame)
 			{
 				Client.CurrentFrame++;
 				Client.StatsWriteFile.ResetData();
-				Client.StatsWriteFile.WriteFrame( Client.CurrentFrame );
+
+				bool bNeedFullMetadata = false;
+				if( Client.MetadataSize < CurrentMetadataSize )
+				{
+					// Write the whole metadata.
+					bNeedFullMetadata = true;
+					Client.MetadataSize = CurrentMetadataSize;
+				}
+
+				Client.StatsWriteFile.WriteFrame( Client.CurrentFrame, bNeedFullMetadata );
 				MessageEndpoint->Send( new FProfilerServiceData2( InstanceId, Client.CurrentFrame, Client.StatsWriteFile.GetOutData() ), PreviewClients );
 			}
 		}
diff --git a/Engine/Source/Developer/ProfilerService/Private/ProfilerServiceManager.h b/Engine/Source/Developer/ProfilerService/Private/ProfilerServiceManager.h
--- a/Engine/Source/Developer/ProfilerService/Private/ProfilerServiceManager.h
+++ b/Engine/Source/Developer/ProfilerService/Private/ProfilerServiceManager.h
@@ -22,6 +22,9 @@ struct FClientData
 
 	/** Current frame of data sent */
 	int64 CurrentFrame;
+
+	/** Stats metadata size. */
+	int32 MetadataSize;
 #endif
 };
 
@@ -55,10 +58,6 @@ public:
 
 	virtual void SendData(FProfilerCycleGraph& Data) OVERRIDE;
 
-	virtual void AddStatDescription(FStatDescription& Data) OVERRIDE;
-
-	virtual void AddGroupDescription(FStatGroupDescription& Data) OVERRIDE;
-
 	virtual bool IsCapturing() const OVERRIDE
 	{ 
 #if STATS
diff --git a/Engine/Source/Developer/ProfilerService/Public/Interfaces/IProfilerServiceManager.h b/Engine/Source/Developer/ProfilerService/Public/Interfaces/IProfilerServiceManager.h
--- a/Engine/Source/Developer/ProfilerService/Public/Interfaces/IProfilerServiceManager.h
+++ b/Engine/Source/Developer/ProfilerService/Public/Interfaces/IProfilerServiceManager.h
@@ -494,20 +494,6 @@ public:
 	 */
 	virtual void SendData(FProfilerCycleGraph& Data) = 0;
 
-	/**
-	 * Specifies the stat description
-	 *
-	 * @param Desc - description of the statistic
-	 */
-	virtual void AddStatDescription(FStatDescription& Data) = 0;
-
-	/**
-	 * Specifies the group description
-	 *
-	 * @param Desc - description of the group
-	 */
-	virtual void AddGroupDescription(FStatGroupDescription& Data) = 0;
-
 	/**
 	 * Determines if we are capturing data for clients
 	 */
