ProfilerService - Added comment

diff --git a/Engine/Source/Developer/ProfilerService/Private/ProfilerServiceManager.cpp b/Engine/Source/Developer/ProfilerService/Private/ProfilerServiceManager.cpp
--- a/Engine/Source/Developer/ProfilerService/Private/ProfilerServiceManager.cpp
+++ b/Engine/Source/Developer/ProfilerService/Private/ProfilerServiceManager.cpp
@@ -358,7 +358,7 @@ void FProfilerServiceManager::StartCapture()
 	// fire off the equivalent of the stat startfile command
 	if (!Archive.IsValid())
 	{
-		// @TODO yrx 2014-03-24 Standardize
+		// @TODO yrx 2014-06-05 Needs to be done on the stats thread via the task graph task.
 		FString Filename = CreateProfileFilename( FStatConstants::StatsFileExtension, true );
 		LastStatFilename = FApp::GetInstanceName() + TEXT("_") + Filename;
 		TSharedPtr<FStatsWriteFile, ESPMode::ThreadSafe> ArchivePtr = MakeShareable(new FStatsWriteFile());
