Profiler - UE-10697 - Crash when opening Post Process Material with Live Data Preview running in Profiler

diff --git a/Engine/Source/Developer/ProfilerClient/Private/ProfilerClientManager.cpp b/Engine/Source/Developer/ProfilerClient/Private/ProfilerClientManager.cpp
--- a/Engine/Source/Developer/ProfilerClient/Private/ProfilerClientManager.cpp
+++ b/Engine/Source/Developer/ProfilerClient/Private/ProfilerClientManager.cpp
@@ -669,10 +669,14 @@ bool FProfilerClientManager::HandleMessagesTicker( float DeltaTime )
 	{
 		FServiceConnection& Connection = It.Value();
 
-		while (Connection.PendingMessages.Contains(Connection.CurrentFrame+1))
+		TArray<int64> Frames;
+		Connection.PendingMessages.GenerateKeyArray( Frames );
+		Frames.Sort();
+
+		for( int32 Index = 0; Index < Frames.Num(); Index++ )
 		{
-			TArray<uint8>& Data = *Connection.PendingMessages.Find(Connection.CurrentFrame+1);
-			Connection.CurrentFrame++;
+			int64 FrameNum = Frames[Index];
+			TArray<uint8>& Data = *Connection.PendingMessages.Find( FrameNum );
 
 			// pass the data to the visualization code
 			FArrayReader Reader(true);
@@ -710,7 +714,7 @@ bool FProfilerClientManager::HandleMessagesTicker( float DeltaTime )
 			// send the data out
 			ProfilerDataDelegate.Broadcast(Connection.InstanceId, Connection.CurrentData,0.0f);
 
-			Connection.PendingMessages.Remove(Connection.CurrentFrame);
+			Connection.PendingMessages.Remove( FrameNum );
 		}
 	}
 
@@ -744,7 +748,6 @@ void FProfilerClientManager::HandleServicePreviewAckMessage( const FProfilerServ
 	if (ActiveSessionId.IsValid() && Connections.Find(Message.InstanceId) != NULL)
 	{
 		FServiceConnection& Connection = *Connections.Find(Message.InstanceId);
-		Connection.CurrentFrame = Message.Frame;
 	}
 #endif
 }
diff --git a/Engine/Source/Developer/ProfilerClient/Private/ProfilerClientManager.h b/Engine/Source/Developer/ProfilerClient/Private/ProfilerClientManager.h
--- a/Engine/Source/Developer/ProfilerClient/Private/ProfilerClientManager.h
+++ b/Engine/Source/Developer/ProfilerClient/Private/ProfilerClientManager.h
@@ -14,7 +14,6 @@ struct FServiceConnection
 		, ReadMessages(0)
 	{
 		MetaData.CriticalSection = &CriticalSection;
-		CurrentFrame = 0;
 	}
 
 	~FServiceConnection()
@@ -25,7 +24,6 @@ struct FServiceConnection
 	FServiceConnection(const FServiceConnection& InConnnection)
 	{
 		MetaData.CriticalSection = &CriticalSection;
-		CurrentFrame = 0;
 	}
 
 	/** Instance Id */
@@ -68,9 +66,6 @@ struct FServiceConnection
 	/** Current data loading progress. */
 	float DataLoadingProgress;
 
-	/** Current frame of data to process */
-	int64 CurrentFrame;
-
 	/** Current number of read messages. */
 	uint64 ReadMessages;
 
diff --git a/Engine/Source/Developer/ProfilerMessages/Classes/ProfilerServiceMessages.h b/Engine/Source/Developer/ProfilerMessages/Classes/ProfilerServiceMessages.h
--- a/Engine/Source/Developer/ProfilerMessages/Classes/ProfilerServiceMessages.h
+++ b/Engine/Source/Developer/ProfilerMessages/Classes/ProfilerServiceMessages.h
@@ -199,10 +199,6 @@ struct FProfilerServicePreviewAck
 	UPROPERTY()
 	FGuid InstanceId;
 
-	UPROPERTY()
-	int64 Frame;
-
-
 	/**
 	 * Default constructor.
 	 */
@@ -211,9 +207,8 @@ struct FProfilerServicePreviewAck
 	/**
 	 * Creates and initializes a new instance.
 	 */
-	FProfilerServicePreviewAck( const FGuid& InInstance, int64 InFrame )
+	FProfilerServicePreviewAck( const FGuid& InInstance )
 		: InstanceId(InInstance)
-		, Frame(InFrame)
 	{
 	}
 };
diff --git a/Engine/Source/Developer/ProfilerService/Private/ProfilerServiceManager.cpp b/Engine/Source/Developer/ProfilerService/Private/ProfilerServiceManager.cpp
--- a/Engine/Source/Developer/ProfilerService/Private/ProfilerServiceManager.cpp
+++ b/Engine/Source/Developer/ProfilerService/Private/ProfilerServiceManager.cpp
@@ -496,8 +496,7 @@ void FProfilerServiceManager::SetPreviewState( const FMessageAddress& ClientAddr
 				Client->Preview = true;
 				if (MessageEndpoint.IsValid())
 				{
-					Client->CurrentFrame = FStats::GameThreadStatsFrame;
-					MessageEndpoint->Send( new FProfilerServicePreviewAck( InstanceId, FStats::GameThreadStatsFrame ), ClientAddress );
+					MessageEndpoint->Send( new FProfilerServicePreviewAck( InstanceId ), ClientAddress );
 				}
 				SendMetaData(ClientAddress);
 			}
@@ -709,22 +708,17 @@ void FProfilerServiceManager::HandleNewFrame(int64 Frame)
 		{
 			FClientData& Client = *ClientData.Find(*It);
 
-			while(Client.CurrentFrame < Frame)
+			Client.StatsWriteFile.ResetData();
+			bool bNeedFullMetadata = false;
+			if( Client.MetadataSize < CurrentMetadataSize )
 			{
-				Client.CurrentFrame++;
-				Client.StatsWriteFile.ResetData();
-
-				bool bNeedFullMetadata = false;
-				if( Client.MetadataSize < CurrentMetadataSize )
-				{
-					// Write the whole metadata.
-					bNeedFullMetadata = true;
-					Client.MetadataSize = CurrentMetadataSize;
-				}
-
-				Client.StatsWriteFile.WriteFrame( Client.CurrentFrame, bNeedFullMetadata );
-				MessageEndpoint->Send( new FProfilerServiceData2( InstanceId, Client.CurrentFrame, Client.StatsWriteFile.GetOutData() ), PreviewClients );
+				// Write the whole metadata.
+				bNeedFullMetadata = true;
+				Client.MetadataSize = CurrentMetadataSize;
 			}
+
+			Client.StatsWriteFile.WriteFrame( Frame, bNeedFullMetadata );
+			MessageEndpoint->Send( new FProfilerServiceData2( InstanceId, Frame, Client.StatsWriteFile.GetOutData() ), PreviewClients );
 		}
 	}
 #endif
diff --git a/Engine/Source/Developer/ProfilerService/Private/ProfilerServiceManager.h b/Engine/Source/Developer/ProfilerService/Private/ProfilerServiceManager.h
--- a/Engine/Source/Developer/ProfilerService/Private/ProfilerServiceManager.h
+++ b/Engine/Source/Developer/ProfilerService/Private/ProfilerServiceManager.h
@@ -16,9 +16,6 @@ struct FClientData
 #if STATS
 	FStatsWriteFile StatsWriteFile;
 
-	/** Current frame of data sent */
-	int64 CurrentFrame;
-
 	/** Stats metadata size. */
 	int32 MetadataSize;
 #endif
