Fix build failures in Test configurations.

diff --git a/Engine/Source/Developer/ProfilerService/Private/ProfilerServiceManager.cpp b/Engine/Source/Developer/ProfilerService/Private/ProfilerServiceManager.cpp
--- a/Engine/Source/Developer/ProfilerService/Private/ProfilerServiceManager.cpp
+++ b/Engine/Source/Developer/ProfilerService/Private/ProfilerServiceManager.cpp
@@ -540,6 +540,7 @@ void FProfilerServiceManager::HandleServiceCaptureMessage( const FProfilerServic
 
 void FProfilerServiceManager::HandleServicePongMessage( const FProfilerServicePong& Message, const IMessageContextRef& Context )
 {
+#if STATS
 	FClientData* Data = ClientData.Find(Context->GetSender());
 	
 	if (Data != nullptr)
@@ -547,6 +548,7 @@ void FProfilerServiceManager::HandleServicePongMessage( const FProfilerServicePo
 		Data->Active = true;
 		UE_LOG( LogProfilerService, Verbose, TEXT( "Pong InstanceId: %s, GetSender: %s" ), *InstanceId.ToString(), *Context->GetSender().ToString() );
 	}
+#endif
 }
 
 
@@ -657,6 +659,7 @@ void FProfilerServiceManager::HandleNewFrame(int64 Frame)
 #endif
 }
 
+#if STATS
 void FProfilerServiceManager::HandleNewFrameGT( FClientData* ToGameThread )
 {
 	if (MessageEndpoint.IsValid())
@@ -668,3 +671,4 @@ void FProfilerServiceManager::HandleNewFrameGT( FClientData* ToGameThread )
 
 	delete ToGameThread;
 }
+#endif
diff --git a/Engine/Source/Developer/ProfilerService/Private/ProfilerServiceManager.h b/Engine/Source/Developer/ProfilerService/Private/ProfilerServiceManager.h
--- a/Engine/Source/Developer/ProfilerService/Private/ProfilerServiceManager.h
+++ b/Engine/Source/Developer/ProfilerService/Private/ProfilerServiceManager.h
@@ -104,9 +104,11 @@ private:
 	/** Handles a new frame from the stats system. Called from the stats thread. */
 	void HandleNewFrame(int64 Frame);
 
+#if STATS
 	/** Handles a new frame from the stats system. Called from the game thread. */
 	void HandleNewFrameGT( FClientData* ToGameThread );
-	
+#endif // STATS
+
 	void AddNewFrameHandleStatsThread();
 
 	void RemoveNewFrameHandleStatsThread();
