Profiler - Optimized data usage, converted TArray<uin8> to FString, requires only 4 time more bandwidth, but not 12, also UDP serializer is no longer an issue, compression/hex is done on the task graph

diff --git a/Engine/Source/Developer/ProfilerClient/Private/ProfilerClientManager.cpp b/Engine/Source/Developer/ProfilerClient/Private/ProfilerClientManager.cpp
--- a/Engine/Source/Developer/ProfilerClient/Private/ProfilerClientManager.cpp
+++ b/Engine/Source/Developer/ProfilerClient/Private/ProfilerClientManager.cpp
@@ -391,7 +391,7 @@ void FServiceConnection::Initialize( const FProfilerServiceAuthorize& Message, c
 	int64 Size = MemoryReader.TotalSize();
 
 	const bool bVerifyHeader = Stream.ReadHeader( MemoryReader );
-	check( bVerifyHeader );
+	//check( bVerifyHeader );
 
 	// Read in the data.
 	TArray<FStatMessage> StatMessages;
@@ -713,14 +713,61 @@ void FProfilerClientManager::HandleProfilerServiceData2Message( const FProfilerS
 	SCOPE_CYCLE_COUNTER(STAT_PC_HandleDataReceived);
 	if (ActiveSessionId.IsValid() && Connections.Find(Message.InstanceId) != nullptr)
 	{
-		FServiceConnection& Connection = *Connections.Find(Message.InstanceId);
+		// Create a temporary profiler data and prepare all data.
+		FProfilerServiceData2* ToProcess = new FProfilerServiceData2( Message.InstanceId, Message.Frame, Message.HexData, Message.CompressedSize, Message.UncompressedSize );
+
+		// Decompression and decoding is done on the task graph.
+		FSimpleDelegateGraphTask::CreateAndDispatchWhenReady
+		(
+			FSimpleDelegateGraphTask::FDelegate::CreateRaw( this, &FProfilerClientManager::DecompressDataAndSendToGame, ToProcess ), 
+			TStatId()
+		);
+	}
+#endif
+}
 
-		// Add the message to the connections queue.
-		Connection.PendingMessages.Add( Message.Frame, Message.Data );
+void FProfilerClientManager::DecompressDataAndSendToGame( FProfilerServiceData2* ToProcess )
+{
+	DECLARE_SCOPE_CYCLE_COUNTER( TEXT( "FProfilerClientManager::DecompressDataAndSendToGame" ), STAT_FProfilerClientManager_DecompressDataAndSendToGame, STATGROUP_Profiler );
+
+	// De-hex string into TArray<uint8>
+	TArray<uint8> CompressedData;
+	CompressedData.Reset( ToProcess->CompressedSize );
+	CompressedData.AddUninitialized( ToProcess->CompressedSize );
+	FString::ToHexBlob( ToProcess->HexData, CompressedData.GetData(), ToProcess->CompressedSize );
+
+	// Decompress data.
+	TArray<uint8> UncompressedData;
+	UncompressedData.Reset( ToProcess->UncompressedSize );
+	UncompressedData.AddUninitialized( ToProcess->UncompressedSize );
+
+	bool bResult = FCompression::UncompressMemory( COMPRESS_ZLIB, UncompressedData.GetData(), ToProcess->UncompressedSize, CompressedData.GetData(), ToProcess->CompressedSize );
+	check( bResult );
+
+	// Send to the game thread. Connections is not thread-safe, so we cannot add the data here.
+	TArray<uint8>* DateToGame = new TArray<uint8>( MoveTemp( UncompressedData ) );
 
-		UE_LOG( LogProfilerClient, VeryVerbose, TEXT( "Frame:%i, Data.Num:%i, InstanceId:%s" ), Message.Frame, Message.Data.Num(), *Message.InstanceId.ToString() );
+	FSimpleDelegateGraphTask::CreateAndDispatchWhenReady
+	(
+		FSimpleDelegateGraphTask::FDelegate::CreateRaw( this, &FProfilerClientManager::SendToGame, DateToGame, ToProcess->Frame, ToProcess->InstanceId ), 
+		TStatId(), nullptr, ENamedThreads::GameThread
+	);
+
+	delete ToProcess;
+}
+
+void FProfilerClientManager::SendToGame( TArray<uint8>* DataToGame, int64 Frame, const FGuid InstanceId )
+{
+	if (ActiveSessionId.IsValid() && Connections.Find( InstanceId ) != nullptr)
+	{
+		FServiceConnection& Connection = *Connections.Find( InstanceId );
+
+		// Add the message to the connections queue.
+		UE_LOG( LogProfilerClient, VeryVerbose, TEXT( "Frame: %i, UncompressedSize: %i, InstanceId: %s" ), Frame, DataToGame->Num(), *InstanceId.ToString() );
+		Connection.PendingMessages.Add( Frame, MoveTemp( *DataToGame ) );
 	}
-#endif
+
+	delete DataToGame;
 }
 
 void FProfilerClientManager::BroadcastMetadataUpdate()
diff --git a/Engine/Source/Developer/ProfilerClient/Private/ProfilerClientManager.h b/Engine/Source/Developer/ProfilerClient/Private/ProfilerClientManager.h
--- a/Engine/Source/Developer/ProfilerClient/Private/ProfilerClientManager.h
+++ b/Engine/Source/Developer/ProfilerClient/Private/ProfilerClientManager.h
@@ -245,6 +245,12 @@ private:
 	/** Broadcast that loading has completed and cleans internal structures. */
 	void FinalizeLoading();
 
+	/** Decompress all stats data and send to the game thread. */
+	void DecompressDataAndSendToGame( FProfilerServiceData2* ToProcess );
+
+	/** Sends decompressed data to the game thread. */
+	void SendToGame( TArray<uint8>* DataToGame, int64 Frame, const FGuid InstanceId );
+
 private:
 
 	/** Session this client is currently communicating with */
diff --git a/Engine/Source/Developer/ProfilerMessages/Classes/ProfilerServiceMessages.h b/Engine/Source/Developer/ProfilerMessages/Classes/ProfilerServiceMessages.h
--- a/Engine/Source/Developer/ProfilerMessages/Classes/ProfilerServiceMessages.h
+++ b/Engine/Source/Developer/ProfilerMessages/Classes/ProfilerServiceMessages.h
@@ -48,7 +48,7 @@ struct FProfilerServiceAuthorize
 /**
  */
 USTRUCT()
-struct FProfilerServiceData
+struct FProfilerServiceData2
 {
 	GENERATED_USTRUCT_BODY()
 
@@ -57,47 +57,18 @@ struct FProfilerServiceData
 	UPROPERTY()
 	FGuid InstanceId;
 
-	/**
-	 */
 	UPROPERTY()
-	TArray<uint8> Data;
-
-
-	/**
-	 * Default constructor.
-	 */
-	FProfilerServiceData( ) { }
-
-	/**
-	 * Creates and initializes a new instance.
-	 */
-	FProfilerServiceData( const FGuid& InInstance, const TArray<uint8>& InData )
-		: InstanceId(InInstance)
-	{
-		Data.Append(InData);
-	}
-};
-
-
-/**
- */
-USTRUCT()
-struct FProfilerServiceData2
-{
-	GENERATED_USTRUCT_BODY()
+	int64 Frame;
 
-	/**
-	 */
 	UPROPERTY()
-	FGuid InstanceId;
-
+	int32 CompressedSize;
+	
 	UPROPERTY()
-	int64 Frame;
+	int32 UncompressedSize;
 
-	/**
-	 */
+	/** Profiler data encoded as string of hexes, cannot use TArray<uint8> because of the Message Bus limitation. */
 	UPROPERTY()
-	TArray<uint8> Data;
+	FString HexData;
 
 
 	/**
@@ -108,11 +79,13 @@ struct FProfilerServiceData2
 	/**
 	 * Creates and initializes a new instance.
 	 */
-	FProfilerServiceData2( const FGuid& InInstance, int64 InFrame, const TArray<uint8>& InData )
-		: InstanceId(InInstance)
-		, Frame(InFrame)
+	FProfilerServiceData2( const FGuid& InInstance, int64 InFrame, const FString& InHexData, int32 InCompressedSize, int32 InUncompressedSize )
+		: InstanceId( InInstance )
+		, Frame( InFrame )
+		, CompressedSize( InCompressedSize )
+		, UncompressedSize( InUncompressedSize )
+		, HexData( MoveTemp( InHexData ) )
 	{
-		Data.Append(InData);
 	}
 };
 
diff --git a/Engine/Source/Developer/ProfilerService/Private/ProfilerServiceManager.cpp b/Engine/Source/Developer/ProfilerService/Private/ProfilerServiceManager.cpp
--- a/Engine/Source/Developer/ProfilerService/Private/ProfilerServiceManager.cpp
+++ b/Engine/Source/Developer/ProfilerService/Private/ProfilerServiceManager.cpp
@@ -583,13 +583,13 @@ void FProfilerServiceManager::HandleServiceSubscribeMessage( const FProfilerServ
 		FClientData Data;
 		Data.Active = true;
 		Data.Preview = false;
-		Data.StatsWriteFile.WriteHeader();
+		//Data.StatsWriteFile.WriteHeader();
 
 		// Add to the client list.
 		ClientData.Add( SenderAddress, Data );
 
 		// Send authorized and stat descriptions.
-		const TArray<uint8>& OutData = ClientData.Find( SenderAddress )->StatsWriteFile.GetOutData();
+		const TArray<uint8> OutData;// = ClientData.Find( SenderAddress )->StatsWriteFile.GetOutData();
 		MessageEndpoint->Send( new FProfilerServiceAuthorize( SessionId, InstanceId, OutData ), SenderAddress );
 
 		// Initiate the ping callback
@@ -633,6 +633,8 @@ void FProfilerServiceManager::HandleNewFrame(int64 Frame)
 {
 	// Called from the stats thread.
 #if STATS
+	DECLARE_SCOPE_CYCLE_COUNTER( TEXT( "FProfilerServiceManager::HandleNewFrame" ), STAT_FProfilerServiceManager_HandleNewFrame, STATGROUP_Profiler );
+	
 	const FStatsThreadState& Stats = FStatsThreadState::GetLocalState();
 	const int32 CurrentMetadataSize = Stats.ShortNameToLongName.Num();
 
@@ -644,11 +646,47 @@ void FProfilerServiceManager::HandleNewFrame(int64 Frame)
 		MetadataSize = CurrentMetadataSize;
 	}
 
-	// Create a temporary client data and prepare all data.
-	FClientData* ToGameThread = new FClientData;
-	ToGameThread->StatsWriteFile.ResetData();
-	ToGameThread->StatsWriteFile.WriteFrame( Frame, bNeedFullMetadata );
-	ToGameThread->Frame = Frame;
+	// Write frame.
+	FStatsWriteFile StatsWriteFile;
+	StatsWriteFile.WriteFrame( Frame, bNeedFullMetadata );
+
+	// Task graph
+	TArray<uint8>* DataToTask = new TArray<uint8>( MoveTemp( StatsWriteFile.GetOutData() ) );
+
+	// Compression and encoding is done on the task graph
+	FSimpleDelegateGraphTask::CreateAndDispatchWhenReady
+	(
+		FSimpleDelegateGraphTask::FDelegate::CreateRaw( this, &FProfilerServiceManager::CompressDataAndSendToGame, DataToTask, Frame ),
+		TStatId()
+	);
+#endif
+}
+
+#if STATS
+
+void FProfilerServiceManager::CompressDataAndSendToGame( TArray<uint8>* DataToTask, int64 Frame )
+{
+	DECLARE_SCOPE_CYCLE_COUNTER( TEXT( "FProfilerServiceManager::CompressDataAndSendToGame" ), STAT_FProfilerServiceManager_CompressDataAndSendToGame, STATGROUP_Profiler );
+
+	const uint8* UncompressedPtr = DataToTask->GetData();
+	const int32 UncompressedSize = DataToTask->Num();
+
+	TArray<uint8> CompressedBuffer;
+	CompressedBuffer.Reserve( UncompressedSize );
+	int32 CompressedSize = UncompressedSize;
+
+	// We assume that compression cannot fail.
+	const bool bResult = FCompression::CompressMemory( COMPRESS_ZLIB, CompressedBuffer.GetData(), CompressedSize, UncompressedPtr, UncompressedSize );
+	check( bResult );
+
+	// Convert to hex.
+	FString HexData = FString::FromHexBlob( CompressedBuffer.GetData(), CompressedSize );
+
+	// Create a temporary profiler data and prepare all data.
+	FProfilerServiceData2* ToGameThread = new FProfilerServiceData2( InstanceId, Frame, HexData, CompressedSize, UncompressedSize );
+
+	const float CompressionRatio = (float)UncompressedSize / (float)CompressedSize;
+	UE_LOG( LogProfilerService, VeryVerbose, TEXT( "Frame: %i, UncompressedSize: %i/%f, InstanceId: %i" ), ToGameThread->Frame, UncompressedSize, CompressionRatio, *InstanceId.ToString() );
 
 	// Send to the game thread. PreviewClients is not thread-safe, so we cannot send the data here.
 	FSimpleDelegateGraphTask::CreateAndDispatchWhenReady
@@ -656,19 +694,16 @@ void FProfilerServiceManager::HandleNewFrame(int64 Frame)
 		FSimpleDelegateGraphTask::FDelegate::CreateRaw( this, &FProfilerServiceManager::HandleNewFrameGT, ToGameThread ),
 		TStatId(), nullptr, ENamedThreads::GameThread
 	);
-#endif
+
+	delete DataToTask;
 }
 
-#if STATS
-void FProfilerServiceManager::HandleNewFrameGT( FClientData* ToGameThread )
+void FProfilerServiceManager::HandleNewFrameGT( FProfilerServiceData2* ToGameThread )
 {
 	if (MessageEndpoint.IsValid())
 	{
-		// Update preview clients with the current data.
-		UE_LOG( LogProfilerService, VeryVerbose, TEXT( "Frame: %i, MetadataSize: %i, PreviewClients: %i" ), ToGameThread->Frame, MetadataSize, PreviewClients.Num() );
-		MessageEndpoint->Send( new FProfilerServiceData2( InstanceId, ToGameThread->Frame, ToGameThread->StatsWriteFile.GetOutData() ), PreviewClients );
+		// Send through the Message Bus.
+		MessageEndpoint->Send( ToGameThread, PreviewClients );
 	}
-
-	delete ToGameThread;
 }
 #endif
diff --git a/Engine/Source/Developer/ProfilerService/Private/ProfilerServiceManager.h b/Engine/Source/Developer/ProfilerService/Private/ProfilerServiceManager.h
--- a/Engine/Source/Developer/ProfilerService/Private/ProfilerServiceManager.h
+++ b/Engine/Source/Developer/ProfilerService/Private/ProfilerServiceManager.h
@@ -10,12 +10,6 @@
  */
 struct FClientData
 {
-	/** Stats data. */
-	FStatsWriteFile StatsWriteFile;
-
-	/** Frame. */
-	int64 Frame;
-
 	/** Connection is active. */
 	bool Active;
 
@@ -24,8 +18,7 @@ struct FClientData
 
 	/** Default constructor. */
 	FClientData()
-		: Frame(-1)
-		, Active( false )
+		: Active( false )
 		, Preview( false )
 	{}
 };
@@ -105,8 +98,11 @@ private:
 	void HandleNewFrame(int64 Frame);
 
 #if STATS
+	/** Compresses all stats data and send to the game thread. */
+	void CompressDataAndSendToGame( TArray<uint8>* DataToTask, int64 Frame );
+
 	/** Handles a new frame from the stats system. Called from the game thread. */
-	void HandleNewFrameGT( FClientData* ToGameThread );
+	void HandleNewFrameGT( FProfilerServiceData2* ToGameThread );
 #endif // STATS
 
 	void AddNewFrameHandleStatsThread();
