UE-4235/UE4: Profiler cannot perform 2 captures in a row through UI (fixed filetransfer runnable not working after recent message bus changes)

diff --git a/Engine/Source/Developer/Profiler/Private/ProfilerCommands.cpp b/Engine/Source/Developer/Profiler/Private/ProfilerCommands.cpp
--- a/Engine/Source/Developer/Profiler/Private/ProfilerCommands.cpp
+++ b/Engine/Source/Developer/Profiler/Private/ProfilerCommands.cpp
@@ -47,7 +47,7 @@ void FProfilerCommands::RegisterCommands()
 
 	UI_COMMAND( OpenSettings, "Settings", "Opens the settings for the profiler", EUserInterfaceActionType::Button, FInputChord( EModifierKey::Control, EKeys::O ) );
 
-	UI_COMMAND( ProfilerManager_Load, "Load", "", EUserInterfaceActionType::Button, FInputChord(/* EModifierKey::Control, EKeys::L */) );
+	UI_COMMAND( ProfilerManager_Load, "Load", "Loads profiler data", EUserInterfaceActionType::Button, FInputChord( EModifierKey::Control, EKeys::L ) );
 	UI_COMMAND( ProfilerManager_ToggleLivePreview, "Live preview", "Toggles the real time live preview", EUserInterfaceActionType::ToggleButton, FInputChord() );
 
 	UI_COMMAND( DataGraph_ToggleViewMode, "Toggle graph view mode", "Toggles the data graph view mode between time based and index based", EUserInterfaceActionType::Button, FInputChord() );
diff --git a/Engine/Source/Developer/ProfilerClient/Private/ProfilerClientManager.cpp b/Engine/Source/Developer/ProfilerClient/Private/ProfilerClientManager.cpp
--- a/Engine/Source/Developer/ProfilerClient/Private/ProfilerClientManager.cpp
+++ b/Engine/Source/Developer/ProfilerClient/Private/ProfilerClientManager.cpp
@@ -422,9 +422,16 @@ bool FProfilerClientManager::CheckHashAndWrite( const FProfilerServiceFileChunk&
 	const int32 HashSize = 20;
 	uint8 LocalHash[HashSize]={0};
 	
+	// De-hex string into TArray<uint8>
+	TArray<uint8> FileChunkData;
+	const int32 DataLength = FileChunk.HexData.Len() / 2;
+	FileChunkData.Reset( DataLength );
+	FileChunkData.AddUninitialized( DataLength );
+	FString::ToHexBlob( FileChunk.HexData, FileChunkData.GetData(), DataLength );
+
 	// Hash file chunk data. 
 	FSHA1 Sha;
-	Sha.Update( FileChunk.Data.GetData(), FileChunkHeader.ChunkSize );
+	Sha.Update( FileChunkData.GetData(), FileChunkHeader.ChunkSize );
 	// Hash file chunk header.
 	Sha.Update( FileChunk.Header.GetData(), FileChunk.Header.Num() );
 	Sha.Final();
@@ -438,7 +445,7 @@ bool FProfilerClientManager::CheckHashAndWrite( const FProfilerServiceFileChunk&
 	{
 		// Write the data to the archive.
 		Writer->Seek( FileChunkHeader.ChunkOffset );
-		Writer->Serialize( (void*)FileChunk.Data.GetData(), FileChunkHeader.ChunkSize );
+		Writer->Serialize( (void*)FileChunkData.GetData(), FileChunkHeader.ChunkSize );
 
 		bResult = true;
 	}
@@ -462,7 +469,6 @@ void FProfilerClientManager::HandleServiceFileChunk( const FProfilerServiceFileC
 
 	const bool bValidFileChunk = !FailedTransfer.Contains( FileChunk.Filename );
 
-	// @TODO yrx 2014-03-24 At this moment received file chunks are handled on the main thread, asynchronous file receiving is planned for the future release.
 	if (ActiveSessionId.IsValid() && Connections.Find(FileChunk.InstanceId) != nullptr && bValidFileChunk)
 	{
 		FReceivedFileInfo* ReceivedFileInfo = ActiveTransfers.Find( FileChunk.Filename );
@@ -522,7 +528,7 @@ void FProfilerClientManager::HandleServiceFileChunk( const FProfilerServiceFileC
 			if( MessageEndpoint.IsValid() )
 			{
 				MessageEndpoint->Send( new FProfilerServiceFileChunk(FileChunk,FProfilerServiceFileChunk::FNullTag()), Context->GetSender() );
-				UE_LOG( LogProfilerClient, Log, TEXT( "Received a bad chunk of file, resending: %5i, %6u, %10u, %s" ), FileChunk.Data.Num(), ReceivedFileInfo->Progress, FileChunkHeader.FileSize, *FileChunk.Filename );
+				UE_LOG( LogProfilerClient, Log, TEXT( "Received a bad chunk of file, resending: %5i, %6u, %10u, %s" ), FileChunk.HexData.Len(), ReceivedFileInfo->Progress, FileChunkHeader.FileSize, *FileChunk.Filename );
 			}
 		}
 	}
diff --git a/Engine/Source/Developer/ProfilerMessages/Classes/ProfilerServiceMessages.h b/Engine/Source/Developer/ProfilerMessages/Classes/ProfilerServiceMessages.h
--- a/Engine/Source/Developer/ProfilerMessages/Classes/ProfilerServiceMessages.h
+++ b/Engine/Source/Developer/ProfilerMessages/Classes/ProfilerServiceMessages.h
@@ -134,9 +134,9 @@ struct FProfilerServiceFileChunk
 	UPROPERTY()
 	FString Filename;
 
-	/** Data to be sent through message bus. */
+	/** Data to be sent through message bus. Message bug doesn't support TArray<>, so we encode the data as HexString. */
 	UPROPERTY()
-	TArray<uint8> Data;
+	FString HexData;
 
 	/** FProfilerFileChunkHeader stored in the array. */
 	UPROPERTY()
diff --git a/Engine/Source/Developer/ProfilerService/Private/ProfilerServiceManager.cpp b/Engine/Source/Developer/ProfilerService/Private/ProfilerServiceManager.cpp
--- a/Engine/Source/Developer/ProfilerService/Private/ProfilerServiceManager.cpp
+++ b/Engine/Source/Developer/ProfilerService/Private/ProfilerServiceManager.cpp
@@ -15,10 +15,6 @@ DEFINE_LOG_CATEGORY_STATIC( LogProfilerService, Log, All );
  */
 class FFileTransferRunnable : public FRunnable
 {
-	/** Archive used to read a captured stats file. Created on the main thread, destroyed on the runnable thread once finalized. */
-	// FArchive* Reader;
-	/** Where this file chunk should be sent. */
-	// FMessageAddress RecipientAddress;
 	typedef TKeyValuePair<FArchive*,FMessageAddress> FReaderAndAddress;
 
 public:	
@@ -127,22 +123,19 @@ public:
 
 	void EnqueueFileToSend( const FString& StatFilename, const FMessageAddress& RecipientAddress, const FGuid& ServiceInstanceId )
 	{
-		const FString PathName = FPaths::ProfilingDir() + TEXT("UnrealStats/");
-		FString StatFilepath = PathName + StatFilename;
+		UE_LOG( LogProfilerService, Log, TEXT( "Opening stats file for service-client sending: %s" ), *StatFilename );
 
-		UE_LOG( LogProfilerService, Log, TEXT( "Opening stats file for service-client sending: %s" ), *StatFilepath );
-
-		const int64 FileSize = IFileManager::Get().FileSize(*StatFilepath);
+		const int64 FileSize = IFileManager::Get().FileSize(*StatFilename);
 		if( FileSize < 4 )
 		{
-			UE_LOG( LogProfilerService, Error, TEXT( "Could not open: %s" ), *StatFilepath );
+			UE_LOG( LogProfilerService, Error, TEXT( "Could not open: %s" ), *StatFilename );
 			return;
 		}
 
-		FArchive* FileReader = IFileManager::Get().CreateFileReader(*StatFilepath);
+		FArchive* FileReader = IFileManager::Get().CreateFileReader(*StatFilename);
 		if( !FileReader )
 		{
-			UE_LOG( LogProfilerService, Error, TEXT( "Could not open: %s" ), *StatFilepath );
+			UE_LOG( LogProfilerService, Error, TEXT( "Could not open: %s" ), *StatFilename );
 			return;
 		}
 
@@ -248,17 +241,19 @@ protected:
 	/** Reads the data from the archive and generates hash. */
 	void ReadAndSetHash( FProfilerServiceFileChunk* FileChunk, const FProfilerFileChunkHeader& FileChunkHeader, FArchive* Reader )
 	{
-		FileChunk->Data.AddUninitialized( FileChunkHeader.ChunkSize );
+		TArray<uint8> FileChunkData;
+
+		FileChunkData.AddUninitialized( FileChunkHeader.ChunkSize );
 
 		Reader->Seek( FileChunkHeader.ChunkOffset );
-		Reader->Serialize( FileChunk->Data.GetData(), FileChunkHeader.ChunkSize );
+		Reader->Serialize( FileChunkData.GetData(), FileChunkHeader.ChunkSize );
 
 		const int32 HashSize = 20;
 		uint8 LocalHash[HashSize]={0};
 
 		// Hash file chunk data. 
 		FSHA1 Sha;
-		Sha.Update( FileChunk->Data.GetData(), FileChunkHeader.ChunkSize );
+		Sha.Update( FileChunkData.GetData(), FileChunkHeader.ChunkSize );
 		// Hash file chunk header.
 		Sha.Update( FileChunk->Header.GetData(), FileChunk->Header.Num() );
 		Sha.Final();
@@ -275,6 +270,9 @@ protected:
 		static const int64 NumBytesPerTick = 256*1024;
 #endif // _DEBUG
 
+		// Convert to hex.
+		FileChunk->HexData = FString::FromHexBlob( FileChunkData.GetData(), FileChunkData.Num() );
+
 		TotalReadBytes += FileChunkHeader.ChunkSize;
 		if( TotalReadBytes > NumBytesPerTick )
 		{
