Copying //UE4/Dev-Sequencer to //UE4/Dev-Main (Source: //UE4/Dev-Sequencer @ 3178529)

diff --git a/Engine/Source/Developer/ProfilerService/Private/ProfilerServiceManager.h b/Engine/Source/Developer/ProfilerService/Private/ProfilerServiceManager.h
--- a/Engine/Source/Developer/ProfilerService/Private/ProfilerServiceManager.h
+++ b/Engine/Source/Developer/ProfilerService/Private/ProfilerServiceManager.h
@@ -2,21 +2,32 @@
 
 #pragma once
 
+#include "IProfilerServiceManager.h"
+#include "Runnable.h"
 
-DECLARE_LOG_CATEGORY_EXTERN( LogProfilerService, Log, All );
+
+struct FMessageAddress;
+class FMessageEndpoint;
+struct FProfilerFileChunkHeader;
+struct FProfilerServiceFileChunk;
+class IMessageContext;
+
+
+DECLARE_LOG_CATEGORY_EXTERN(LogProfilerService, Log, All);
 
 /**
 * Thread used to read, prepare and send files through the message bus.
 * Supports resending bad file chunks and basic synchronization between service and client.
 */
-class FFileTransferRunnable : public FRunnable
+class FFileTransferRunnable
+	: public FRunnable
 {
 	typedef TKeyValuePair<FArchive*, FMessageAddress> FReaderAndAddress;
 
 public:
 
 	/** Default constructor. */
-	FFileTransferRunnable( FMessageEndpointPtr& InMessageEndpoint );
+	FFileTransferRunnable(TSharedPtr<FMessageEndpoint, ESPMode::ThreadSafe>& InMessageEndpoint);
 
 	/** Destructor. */
 	~FFileTransferRunnable();
@@ -34,19 +45,19 @@ public:
 	virtual void Exit();
 	// End FRunnable interface
 
-	void EnqueueFileToSend( const FString& StatFilename, const FMessageAddress& RecipientAddress, const FGuid& ServiceInstanceId );
+	void EnqueueFileToSend(const FString& StatFilename, const FMessageAddress& RecipientAddress, const FGuid& ServiceInstanceId);
 
 	/** Enqueues a file chunk. */
-	void EnqueueFileChunkToSend( FProfilerServiceFileChunk* FileChunk, bool bTriggerWorkEvent = false );
+	void EnqueueFileChunkToSend(FProfilerServiceFileChunk* FileChunk, bool bTriggerWorkEvent = false);
 
 	/** Prepare the chunks to be sent through the message bus. */
-	void PrepareFileForSending( FProfilerServiceFileChunk*& FileChunk );
+	void PrepareFileForSending(FProfilerServiceFileChunk*& FileChunk);
 
 	/** Removes file from the list of the active transfers, must be confirmed by the profiler client. */
-	void FinalizeFileSending( const FString& Filename );
+	void FinalizeFileSending(const FString& Filename);
 
 	/** Aborts file sending to the specified client, probably client disconnected or exited. */
-	void AbortFileSending( const FMessageAddress& Recipient );
+	void AbortFileSending(const FMessageAddress& Recipient);
 
 	/** Checks if there has been any stop requests. */
 	FORCEINLINE bool ShouldStop() const
@@ -55,11 +66,12 @@ public:
 	}
 
 protected:
+
 	/** Deletes the file reader. */
-	void DeleteFileReader( FReaderAndAddress& ReaderAndAddress );
+	void DeleteFileReader(FReaderAndAddress& ReaderAndAddress);
 
 	/** Reads the data from the archive and generates hash. */
-	void ReadAndSetHash( FProfilerServiceFileChunk* FileChunk, const FProfilerFileChunkHeader& FileChunkHeader, FArchive* Reader );
+	void ReadAndSetHash(FProfilerServiceFileChunk* FileChunk, const FProfilerFileChunkHeader& FileChunkHeader, FArchive* Reader);
 
 	/** Thread that is running this task. */
 	FRunnableThread* Runnable;
@@ -83,6 +95,7 @@ protected:
 	TMap<FString, FReaderAndAddress> ActiveTransfers;
 };
 
+
 #if STATS
 
 /**
@@ -98,9 +111,9 @@ struct FClientData
 
 	/** Default constructor. */
 	FClientData()
-		: Active( false )
-		, Preview( false )
-	{}
+		: Active(false)
+		, Preview(false)
+	{ }
 };
 
 #endif //STATS
@@ -115,75 +128,74 @@ class FProfilerServiceManager
 {
 public:
 
-	/**
-	 * Default constructor
-	 */
+	/** Default constructor. */
 	FProfilerServiceManager();
 
 public:
 
-	//~ Begin IProfilerServiceManager Interface
+	//~ IProfilerServiceManager interface
+
 	virtual void StartCapture() override;
 	virtual void StopCapture() override;
-	//~ End IProfilerServiceManager Interface
+
+public:
 
 	/**
 	 * Creates a profiler service manager for shared use
 	 */
-	static IProfilerServiceManagerPtr CreateSharedServiceManager();
+	static TSharedPtr<IProfilerServiceManager> CreateSharedServiceManager();
 
 
-	/**
-	 * Initializes the manager
-	 */
+	/** Initializes the manager. */
 	void Init();
 
-	/**
-	 * Shuts down the manager
-	 */
+	/** Shuts down the manager. */
 	void Shutdown();
 
 private:
+
 	/**
 	 * Changes the data preview state for the given client to the specified value.
 	 */
-	void SetPreviewState( const FMessageAddress& ClientAddress, const bool bRequestedPreviewState );
+	void SetPreviewState(const FMessageAddress& ClientAddress, const bool bRequestedPreviewState);
 
 	/** Callback for a tick, used to ping the clients */
-	bool HandlePing( float DeltaTime );
+	bool HandlePing(float DeltaTime);
 
 
-	// Handles FProfilerServiceCapture messages.
-	void HandleServiceCaptureMessage( const FProfilerServiceCapture& Message, const IMessageContextRef& Context );
+	/** Handles FProfilerServiceCapture messages. */
+	void HandleServiceCaptureMessage(const FProfilerServiceCapture& Message, const TSharedRef<IMessageContext, ESPMode::ThreadSafe>& Context);
 
-	// Handles FProfilerServicePong messages.
-	void HandleServicePongMessage( const FProfilerServicePong& Message, const IMessageContextRef& Context );
+	/** Handles FProfilerServicePong messages. */
+	void HandleServicePongMessage(const FProfilerServicePong& Message, const TSharedRef<IMessageContext, ESPMode::ThreadSafe>& Context);
 
-	// Handles FProfilerServicePreview messages.
-	void HandleServicePreviewMessage( const FProfilerServicePreview& Message, const IMessageContextRef& Context );
+	/** Handles FProfilerServicePreview messages. */
+	void HandleServicePreviewMessage(const FProfilerServicePreview& Message, const TSharedRef<IMessageContext, ESPMode::ThreadSafe>& Context);
 
-	// Handles FProfilerServiceRequest messages.
-	void HandleServiceRequestMessage( const FProfilerServiceRequest& Message, const IMessageContextRef& Context );
+	/** Handles FProfilerServiceRequest messages. */
+	void HandleServiceRequestMessage(const FProfilerServiceRequest& Message, const TSharedRef<IMessageContext, ESPMode::ThreadSafe>& Context);
 
-	// Handles FProfilerServiceFileChunk messages.
-	void HandleServiceFileChunkMessage( const FProfilerServiceFileChunk& Message, const IMessageContextRef& Context );
+	/** Handles FProfilerServiceFileChunk messages. */
+	void HandleServiceFileChunkMessage(const FProfilerServiceFileChunk& Message, const TSharedRef<IMessageContext, ESPMode::ThreadSafe>& Context);
 
-	// Handles FProfilerServiceSubscribe messages.
-	void HandleServiceSubscribeMessage( const FProfilerServiceSubscribe& Message, const IMessageContextRef& Context );
+	/** Handles FProfilerServiceSubscribe messages. */
+	void HandleServiceSubscribeMessage(const FProfilerServiceSubscribe& Message, const TSharedRef<IMessageContext, ESPMode::ThreadSafe>& Context);
 
-	// Handles FProfilerServiceUnsubscribe messages.
-	void HandleServiceUnsubscribeMessage( const FProfilerServiceUnsubscribe& Message, const IMessageContextRef& Context );
+	/** Handles FProfilerServiceUnsubscribe messages. */
+	void HandleServiceUnsubscribeMessage(const FProfilerServiceUnsubscribe& Message, const TSharedRef<IMessageContext, ESPMode::ThreadSafe>& Context);
 
 	/** Handles a new frame from the stats system. Called from the stats thread. */
 	void HandleNewFrame(int64 Frame);
 
 #if STATS
+
 	/** Compresses all stats data and send to the game thread. */
-	void CompressDataAndSendToGame( TArray<uint8>* DataToTask, int64 Frame );
+	void CompressDataAndSendToGame(TArray<uint8>* DataToTask, int64 Frame);
 
 	/** Handles a new frame from the stats system. Called from the game thread. */
-	void HandleNewFrameGT( FProfilerServiceData2* ToGameThread );
-#endif // STATS
+	void HandleNewFrameGT(FProfilerServiceData2* ToGameThread);
+
+#endif //STATS
 
 	void AddNewFrameHandleStatsThread();
 
@@ -200,9 +212,11 @@ private:
 	TArray<FMessageAddress> PreviewClients;
 
 #if	STATS
+
 	/** Holds the client data for registered clients */
 	TMap<FMessageAddress, FClientData> ClientData;
-#endif // STATS
+
+#endif //STATS
 
 	/** Thread used to read, prepare and send file chunks through the message bus. */
 	FFileTransferRunnable* FileTransferRunnable;
diff --git a/Engine/Source/Developer/ProfilerService/Private/ProfilerServiceModule.cpp b/Engine/Source/Developer/ProfilerService/Private/ProfilerServiceModule.cpp
--- a/Engine/Source/Developer/ProfilerService/Private/ProfilerServiceModule.cpp
+++ b/Engine/Source/Developer/ProfilerService/Private/ProfilerServiceModule.cpp
@@ -1,7 +1,9 @@
 // Copyright 1998-2016 Epic Games, Inc. All Rights Reserved.
 
 #include "ProfilerServicePrivatePCH.h"
-#include "ModuleManager.h"
+#include "ProfilerServiceManager.h"
+#include "IProfilerServiceManager.h"
+#include "IProfilerServiceModule.h"
 
 
 /**
@@ -12,7 +14,8 @@ class FProfilerServiceModule
 {
 public:
 
-	// Begin IModuleInterface interface
+	//~ IModuleInterface interface
+
 	virtual void ShutdownModule() override
 	{
 		if (ProfilerServiceManager.IsValid())
@@ -21,10 +24,12 @@ public:
 		}
 		ProfilerServiceManager.Reset();
 	}
-	// End IModuleInterface interface
 
-	// Begin IProfilerServiceModule interface
-	virtual IProfilerServiceManagerPtr CreateProfilerServiceManager( ) override
+public:
+
+	//~ IProfilerServiceModule interface
+
+	virtual TSharedPtr<IProfilerServiceManager> CreateProfilerServiceManager() override
 	{
 		if (!ProfilerServiceManager.IsValid())
 		{
@@ -34,18 +39,15 @@ public:
 
 		return ProfilerServiceManager;
 	}
-	// End IProfilerServiceModule interface
 
 private:
-	/**
-	 * Holds the profiler service singleton
-	 */
-	static IProfilerServiceManagerPtr ProfilerServiceManager;
+
+	/** The profiler service singleton. */
+	static TSharedPtr<IProfilerServiceManager> ProfilerServiceManager;
 };
 
-/* Static initialization
-******************************************************************************/
-IProfilerServiceManagerPtr FProfilerServiceModule::ProfilerServiceManager = NULL;
+
+TSharedPtr<IProfilerServiceManager> FProfilerServiceModule::ProfilerServiceManager = nullptr;
 
 
 IMPLEMENT_MODULE(FProfilerServiceModule, ProfilerService);
diff --git a/Engine/Source/Developer/ProfilerService/Private/ProfilerServicePrivatePCH.h b/Engine/Source/Developer/ProfilerService/Private/ProfilerServicePrivatePCH.h
--- a/Engine/Source/Developer/ProfilerService/Private/ProfilerServicePrivatePCH.h
+++ b/Engine/Source/Developer/ProfilerService/Private/ProfilerServicePrivatePCH.h
@@ -3,19 +3,12 @@
 #pragma once
 
 
-/* Dependencies
- *****************************************************************************/
-
-#include "Engine.h"
-#include "Messaging.h"
-#include "ProfilerMessages.h"
-#include "Ticker.h"
-#include "StatsData.h"
-#include "StatsFile.h"
-
-
-/* Private includes
- *****************************************************************************/
-
-#include "ProfilerService.h"
-#include "ProfilerServiceManager.h"
+#include "Runtime/Engine/Public/Engine.h"
+#include "Runtime/Core/Public/Containers/Ticker.h"
+#include "Runtime/Core/Public/Misc/SecureHash.h"
+#include "Runtime/Core/Public/Stats/StatsData.h"
+#include "Runtime/Core/Public/Stats/StatsFile.h"
+#include "Runtime/Messaging/Public/Helpers/MessageEndpoint.h"
+#include "Runtime/Messaging/Public/Helpers/MessageEndpointBuilder.h"
+
+#include "Developer/ProfilerMessages/Public/ProfilerServiceMessages.h"
diff --git a/Engine/Source/Developer/ProfilerService/Public/IProfilerServiceManager.h b/Engine/Source/Developer/ProfilerService/Public/IProfilerServiceManager.h
new file mode 100644
--- /dev/null
+++ b/Engine/Source/Developer/ProfilerService/Public/IProfilerServiceManager.h
@@ -0,0 +1,345 @@
+// Copyright 1998-2016 Epic Games, Inc. All Rights Reserved.
+
+#pragma once
+
+
+/** Type definition for shared references to instances of IProfilerService. */
+typedef TSharedRef<class IProfilerServiceManager> IProfilerServiceManagerRef;
+
+
+/**
+ * Type definition for the profiler data packet.
+ */
+struct FProfilerCycleCounter
+{
+public:
+	uint32 StatId;
+	uint32 GroupId;
+	double FrameStart;
+	uint32 ParentId;
+	uint32 InstanceId;
+	uint32 ParentInstanceId;
+	uint32 ThreadId;
+	uint32 Value;
+	uint32 CallsPerFrame;
+};
+
+
+struct FProfilerFloatAccumulator
+{
+	int32 StatId;
+	float Value;
+};
+
+
+struct FProfilerCountAccumulator
+{
+	int32 StatId;
+	uint32 Value;
+};
+
+
+struct FProfilerCycleGraph
+{
+public:
+	int32 StatId;
+	uint32 FrameStart;
+	uint32 ThreadId;
+	uint32 Value;
+	uint32 CallsPerFrame;
+	TArray<FProfilerCycleGraph> Children;
+};
+
+
+FORCEINLINE FArchive& operator<<(FArchive& Ar, FProfilerCycleCounter& Data)
+{
+	Ar << Data.StatId;
+	Ar << Data.GroupId;
+	Ar << Data.FrameStart;
+	Ar << Data.ParentId;
+	Ar << Data.InstanceId;
+	Ar << Data.ParentInstanceId;
+	Ar << Data.ThreadId;
+	Ar << Data.Value;
+	Ar << Data.CallsPerFrame;
+
+	return Ar;
+}
+
+	
+FORCEINLINE FArchive& operator<<(FArchive& Ar, FProfilerFloatAccumulator& Data)
+{
+	Ar << Data.StatId;
+	Ar << Data.Value;
+
+	return Ar;
+}
+
+
+FORCEINLINE FArchive& operator<<(FArchive& Ar, FProfilerCountAccumulator& Data)
+{
+	Ar << Data.StatId;
+	Ar << Data.Value;
+
+	return Ar;
+}
+
+
+FORCEINLINE FArchive& operator<<(FArchive& Ar, FProfilerCycleGraph& Data)
+{
+	Ar << Data.StatId;
+
+	Ar << Data.FrameStart;
+	Ar << Data.ThreadId;
+	Ar << Data.CallsPerFrame;
+	Ar << Data.Value;
+	Ar << Data.Children;
+
+	return Ar;
+}
+
+
+/**
+ * Type definition for a frame profiler data
+ */
+struct FProfilerDataFrame
+{
+	TMap<uint32, TArray<FProfilerCycleCounter> > CycleCounters;
+	TMap<uint32, FProfilerCycleGraph> CycleGraphs;
+	TArray<FProfilerFloatAccumulator> FloatAccumulators;
+	TArray<FProfilerCountAccumulator> CountAccumulators;
+
+	double FrameStart;
+	uint32 Frame;
+	bool MetaDataUpdated;
+};
+
+FORCEINLINE FArchive& operator<<(FArchive& Ar, FProfilerDataFrame& Data)
+{
+	Ar << Data.Frame;
+	Ar << Data.FrameStart;
+	Ar << Data.CycleCounters;
+	Ar << Data.CycleGraphs;
+	Ar << Data.FloatAccumulators;
+	Ar << Data.CountAccumulators;
+
+	return Ar;
+}
+
+
+struct FStatDescription
+{
+	/** Default constructor. */
+	FStatDescription()
+		: ID( 0 )
+		, StatType( -1 )
+		, GroupID( 0 )
+	{}
+
+	/** ID for this stat */
+	int32 ID;
+
+	/** Stat name */
+	FString Name;
+
+	/** Type of stat (counter, cycle, etc.) */
+	uint32 StatType;
+
+	/** Group this stat belongs to */
+	int32 GroupID;
+};
+
+FORCEINLINE FArchive& operator<<(FArchive& Ar, FStatDescription& Data)
+{
+	Ar << Data.ID;
+	Ar << Data.Name;
+	Ar << Data.StatType;
+	Ar << Data.GroupID;
+
+	return Ar;
+}
+	
+
+struct FStatGroupDescription
+{
+	/** Default constructor. */
+	FStatGroupDescription()
+		: ID( 0 )
+	{}
+
+	/** ID for this group */
+	int32 ID;
+
+	/** Group name */
+	FString Name;
+
+};
+
+FORCEINLINE FArchive& operator<<(FArchive& Ar, FStatGroupDescription& Data)
+{
+	Ar << Data.ID;
+	Ar << Data.Name;
+
+	return Ar;
+}
+
+
+/**
+ * Structure holding the meta data describing the various stats and data associated with them
+ */
+struct FStatMetaData
+{
+	/** Stat descriptions */
+	TMap<uint32, FStatDescription> StatDescriptions;
+
+	/** Group descriptions */
+	TMap<uint32, FStatGroupDescription> GroupDescriptions;
+
+	/** Thread descriptions */
+	TMap<uint32, FString> ThreadDescriptions;
+
+	/** Seconds per cycle */
+	double SecondsPerCycle;
+
+	/** Critical section used to lock the meta data for access */
+	FCriticalSection* CriticalSection;
+
+	const uint32 GetMetaDataSize() const
+	{
+		return StatDescriptions.Num() + GroupDescriptions.Num() + ThreadDescriptions.Num();
+	}
+
+};
+
+FORCEINLINE FArchive& operator<<(FArchive& Ar, FStatMetaData& Data)
+{
+	Ar << Data.StatDescriptions;
+	Ar << Data.GroupDescriptions;
+	Ar << Data.ThreadDescriptions;
+	Ar << Data.SecondsPerCycle;
+
+	return Ar;
+}
+
+/**
+ * Profiler service request type
+ */
+namespace EProfilerRequestType
+{
+	enum Type
+	{
+		/** Send last captured file. */
+		PRT_SendLastCapturedFile,
+	};
+}
+
+/** Delegate for passing profiler data to UI */
+DECLARE_MULTICAST_DELEGATE_TwoParams(FProfilerDataDelegate, const FGuid&, const FProfilerDataFrame& );
+
+#if STATS
+DECLARE_MULTICAST_DELEGATE_TwoParams(FProfilerData2Delegate, const FGuid&, TArray<FStatMessage>& );
+#endif
+
+
+/** Enumerates profiler file chunk types. */
+namespace EProfilerFileChunkType
+{
+	enum Type
+	{
+		/** Indicates that message should prepare chunks for sending. This is a helper file chunk. */
+		PrepareFile,
+		/** Indicates that message requires sending a file chunk. */
+		SendChunk,
+		/** Indicates that message requires finalization of transfered file. This is a helper file chunk. */
+		FinalizeFile,
+		/** Not used. */
+		NotUsed,
+	};
+}
+
+
+/** Class that describes basic information about one file chunk. */
+struct FProfilerFileChunkHeader
+{
+	/** Size of the file chunk. */
+	static const int64 DefChunkSize = 64*1024;
+
+	/** Default constructor. */
+	FProfilerFileChunkHeader()
+		: ChunkOffset( 0 )
+		, ChunkSize( 0 )
+		, FileSize( 0 )
+		, ChunkType( EProfilerFileChunkType::NotUsed )
+	{}
+
+	/** Initialization constructor. */
+	FProfilerFileChunkHeader( int64 InChunkOffset, int64 InChunkSize, int64 InFileSize, EProfilerFileChunkType::Type InChunkType )
+		: ChunkOffset( InChunkOffset )
+		, ChunkSize( InChunkSize )
+		, FileSize( InFileSize )
+		, ChunkType( InChunkType )
+	{}
+
+	/** Performs sanity checks to make sure that header data is valid. */
+	FORCEINLINE_DEBUGGABLE void Validate() const
+	{
+		check( ChunkOffset % DefChunkSize == 0 );
+		check( ChunkSize <= DefChunkSize  );
+		check( ChunkType == EProfilerFileChunkType::SendChunk || ChunkType == EProfilerFileChunkType::PrepareFile || ChunkType == EProfilerFileChunkType::FinalizeFile );
+		check( FileSize >= 0 );
+		check( ChunkOffset >= 0 );
+	}
+
+	/** Offset of this chunk. */
+	int64 ChunkOffset;
+
+	/** Size of this chunk. */
+	int64 ChunkSize;
+
+	/** The size of the file. */
+	int64 FileSize;
+
+	/** File chunk type @see EProfilerFileChunkType. */
+	EProfilerFileChunkType::Type ChunkType;
+
+	friend FORCEINLINE FArchive& operator<<(FArchive& Ar, FProfilerFileChunkHeader& Header)
+	{
+		Ar << Header.ChunkOffset;
+		Ar << Header.ChunkSize;
+		Ar << Header.FileSize;
+		Ar << (uint32&)Header.ChunkType;
+
+		return Ar;
+	}
+
+	/** Serializes this file chunk header as an array of bytes. */
+	FORCEINLINE_DEBUGGABLE TArray<uint8> AsArray() const
+	{
+		TArray<uint8> ArrayHeader;
+
+		// Setup header.
+		FMemoryWriter Writer(ArrayHeader);
+		Writer << const_cast<FProfilerFileChunkHeader&>(*this);
+
+		return ArrayHeader;
+	}
+};
+
+
+/**
+ * Interface for the Profiler Service manager.
+ */
+class IProfilerServiceManager
+{
+public:
+
+	/** Starts a file capture. */
+	virtual void StartCapture() = 0;
+
+	/** Stops a file capture. */
+	virtual void StopCapture() = 0;
+public:
+
+	/** Virtual destructor. */
+	virtual ~IProfilerServiceManager() { }
+};
diff --git a/Engine/Source/Developer/ProfilerService/Public/IProfilerServiceModule.h b/Engine/Source/Developer/ProfilerService/Public/IProfilerServiceModule.h
new file mode 100644
--- /dev/null
+++ b/Engine/Source/Developer/ProfilerService/Public/IProfilerServiceModule.h
@@ -0,0 +1,31 @@
+// Copyright 1998-2016 Epic Games, Inc. All Rights Reserved.
+
+#pragma once
+
+#include "ModuleInterface.h"
+
+
+class IProfilerServiceManager;
+
+
+/**
+ * Interface for ProfilerService modules.
+ */
+class IProfilerServiceModule
+	: public IModuleInterface
+{
+public:
+
+	/** 
+	 * Creates a profiler service.
+	 *
+	 * @param Capture Specifies whether to start capturing immediately.
+	 * @return A new profiler service.
+	 */
+	virtual TSharedPtr<IProfilerServiceManager> CreateProfilerServiceManager() = 0;
+
+protected:
+
+	/** Virtual destructor */
+	virtual ~IProfilerServiceModule() { }
+};
diff --git a/Engine/Source/Developer/SessionFrontend/Private/SessionFrontendModule.cpp b/Engine/Source/Developer/SessionFrontend/Private/SessionFrontendModule.cpp
--- a/Engine/Source/Developer/SessionFrontend/Private/SessionFrontendModule.cpp
+++ b/Engine/Source/Developer/SessionFrontend/Private/SessionFrontendModule.cpp
@@ -19,12 +19,12 @@ public:
 
 	// ISessionFrontendModule interface
 	
-	virtual TSharedRef<class SWidget> CreateSessionBrowser( const ISessionManagerRef& SessionManager ) override
+	virtual TSharedRef<class SWidget> CreateSessionBrowser( const TSharedRef<ISessionManager>& SessionManager ) override
 	{
 		return SNew(SSessionBrowser, SessionManager);
 	}
 	
-	virtual TSharedRef<class SWidget> CreateSessionConsole( const ISessionManagerRef& SessionManager ) override
+	virtual TSharedRef<class SWidget> CreateSessionConsole( const TSharedRef<ISessionManager>& SessionManager ) override
 	{
 		return SNew(SSessionConsole, SessionManager);
 	}
diff --git a/Engine/Source/Developer/SessionFrontend/Private/SessionFrontendPrivatePCH.h b/Engine/Source/Developer/SessionFrontend/Private/SessionFrontendPrivatePCH.h
--- a/Engine/Source/Developer/SessionFrontend/Private/SessionFrontendPrivatePCH.h
+++ b/Engine/Source/Developer/SessionFrontend/Private/SessionFrontendPrivatePCH.h
@@ -20,13 +20,14 @@
 #include "ISessionInfo.h"
 #include "ISessionInstanceInfo.h"
 #include "ISessionManager.h"
+#include "ISessionServicesModule.h"
 #include "SessionLogMessage.h"
 
 // @todo gmp: remove these dependencies by making the session front-end extensible
 #include "AutomationWindow.h"
+#include "IProfilerModule.h"
 #include "ScreenShotComparison.h"
 #include "ScreenShotComparisonTools.h"
-#include "Profiler.h"
 
 
 /* Private includes
diff --git a/Engine/Source/Developer/SessionFrontend/Private/Widgets/Browser/SSessionBrowser.cpp b/Engine/Source/Developer/SessionFrontend/Private/Widgets/Browser/SSessionBrowser.cpp
--- a/Engine/Source/Developer/SessionFrontend/Private/Widgets/Browser/SSessionBrowser.cpp
+++ b/Engine/Source/Developer/SessionFrontend/Private/Widgets/Browser/SSessionBrowser.cpp
@@ -28,7 +28,7 @@ SSessionBrowser::~SSessionBrowser()
  *****************************************************************************/
 
 BEGIN_SLATE_FUNCTION_BUILD_OPTIMIZATION
-void SSessionBrowser::Construct( const FArguments& InArgs, ISessionManagerRef InSessionManager )
+void SSessionBrowser::Construct( const FArguments& InArgs, TSharedRef<ISessionManager> InSessionManager )
 {
 	IgnoreSessionManagerEvents = false;
 	updatingTreeExpansion = false;
@@ -155,7 +155,7 @@ void SSessionBrowser::FilterSessions()
 		NewItemMap.Add(SessionInfo->GetSessionId(), SessionItem);
 
 		// add session to group
-		TArray<ISessionInstanceInfoPtr> Instances;
+		TArray<TSharedPtr<ISessionInstanceInfo>> Instances;
 		SessionInfo->GetInstances(Instances);
 
 		if (LocalOwner)
@@ -258,7 +258,7 @@ void SSessionBrowser::HandleSessionManagerInstanceSelectionChanged(const TShared
 }
 
 
-void SSessionBrowser::HandleSessionManagerSelectedSessionChanged(const ISessionInfoPtr& SelectedSession)
+void SSessionBrowser::HandleSessionManagerSelectedSessionChanged(const TSharedPtr<ISessionInfo>& SelectedSession)
 {
 	if (IgnoreSessionManagerEvents)
 	{
@@ -310,7 +310,7 @@ FText SSessionBrowser::HandleSessionTreeRowGetToolTipText(TSharedPtr<FSessionBro
 
 	if (Item->GetType() == ESessionBrowserTreeNodeType::Instance)
 	{
-		ISessionInstanceInfoPtr InstanceInfo = StaticCastSharedPtr<FSessionBrowserInstanceTreeItem>(Item)->GetInstanceInfo();
+		TSharedPtr<ISessionInstanceInfo> InstanceInfo = StaticCastSharedPtr<FSessionBrowserInstanceTreeItem>(Item)->GetInstanceInfo();
 
 		if (InstanceInfo.IsValid())
 		{
@@ -428,7 +428,7 @@ void SSessionBrowser::HandleSessionTreeViewSelectionChanged(const TSharedPtr<FSe
 
 		{
 			// check if any instances are no longer selected
-			TArray<ISessionInstanceInfoPtr> UnselectedSessions;
+			TArray<TSharedPtr<ISessionInstanceInfo>> UnselectedSessions;
 			for (const auto& InstanceInfo : SessionManager->GetSelectedInstances())
 			{
 				const TSharedPtr<FSessionBrowserTreeItem>& InstanceItem = ItemMap.FindRef(InstanceInfo->GetInstanceId());
@@ -455,7 +455,7 @@ FReply SSessionBrowser::HandleTerminateSessionButtonClicked()
 
 	if (DialogResult == EAppReturnType::Yes)
 	{
-		const ISessionInfoPtr& SelectedSession = SessionManager->GetSelectedSession();
+		const TSharedPtr<ISessionInfo>& SelectedSession = SessionManager->GetSelectedSession();
 
 		if (SelectedSession.IsValid())
 		{
diff --git a/Engine/Source/Developer/SessionFrontend/Private/Widgets/Browser/SSessionBrowser.h b/Engine/Source/Developer/SessionFrontend/Private/Widgets/Browser/SSessionBrowser.h
--- a/Engine/Source/Developer/SessionFrontend/Private/Widgets/Browser/SSessionBrowser.h
+++ b/Engine/Source/Developer/SessionFrontend/Private/Widgets/Browser/SSessionBrowser.h
@@ -27,7 +27,7 @@ public:
 	 * @param InArgs The declaration data for this widget.
 	 * @param InSessionManager The session manager to use.
 	 */
-	void Construct( const FArguments& InArgs, ISessionManagerRef InSessionManager );
+	void Construct( const FArguments& InArgs, TSharedRef<ISessionManager> InSessionManager );
 
 protected:
 
@@ -53,7 +53,7 @@ private:
 	void HandleSessionManagerInstanceSelectionChanged(const TSharedPtr<ISessionInstanceInfo>& Instance, bool Selected);
 
 	/** Callback for changing the selected session in the session manager. */
-	void HandleSessionManagerSelectedSessionChanged(const ISessionInfoPtr& SelectedSession);
+	void HandleSessionManagerSelectedSessionChanged(const TSharedPtr<ISessionInfo>& SelectedSession);
 
 	/** Callback for updating the session list in the session manager. */
 	void HandleSessionManagerSessionsUpdated();
@@ -85,7 +85,7 @@ private:
 private:
 
 	/** Holds an unfiltered list of available sessions. */
-	TArray<ISessionInfoPtr> AvailableSessions;
+	TArray<TSharedPtr<ISessionInfo>> AvailableSessions;
 
 	/** Holds the command bar. */
 	TSharedPtr<SSessionBrowserCommandBar> CommandBar;
@@ -100,7 +100,7 @@ private:
 	TMap<FGuid, TSharedPtr<FSessionBrowserTreeItem>> ItemMap;
 
 	/** Holds a reference to the session manager. */
-	ISessionManagerPtr SessionManager;
+	TSharedPtr<ISessionManager> SessionManager;
 
 	/** Holds the filtered list of tree items. */
 	TArray<TSharedPtr<FSessionBrowserTreeItem>> SessionTreeItems;
diff --git a/Engine/Source/Developer/SessionFrontend/Private/Widgets/Browser/SSessionBrowserTreeInstanceRow.h b/Engine/Source/Developer/SessionFrontend/Private/Widgets/Browser/SSessionBrowserTreeInstanceRow.h
--- a/Engine/Source/Developer/SessionFrontend/Private/Widgets/Browser/SSessionBrowserTreeInstanceRow.h
+++ b/Engine/Source/Developer/SessionFrontend/Private/Widgets/Browser/SSessionBrowserTreeInstanceRow.h
@@ -49,7 +49,7 @@ public:
 	{
 		if (ColumnName == "Device")
 		{
-			ISessionInstanceInfoPtr InstanceInfo = Item->GetInstanceInfo();
+			TSharedPtr<ISessionInstanceInfo> InstanceInfo = Item->GetInstanceInfo();
 
 			if (InstanceInfo.IsValid())
 			{
@@ -94,7 +94,7 @@ public:
 		}
 		else if (ColumnName == "Name")
 		{
-			ISessionInstanceInfoPtr InstanceInfo = Item->GetInstanceInfo();
+			TSharedPtr<ISessionInstanceInfo> InstanceInfo = Item->GetInstanceInfo();
 
 			if (InstanceInfo.IsValid())
 			{
@@ -140,7 +140,7 @@ public:
 		}
 		else if (ColumnName == "Type")
 		{
-			ISessionInstanceInfoPtr InstanceInfo = Item->GetInstanceInfo();
+			TSharedPtr<ISessionInstanceInfo> InstanceInfo = Item->GetInstanceInfo();
 
 			if (InstanceInfo.IsValid())
 			{
@@ -164,7 +164,7 @@ private:
 	/** Callback for getting the image of the Authorized icon. */
 	const FSlateBrush* HandleAuthorizedImage() const
 	{
-		ISessionInstanceInfoPtr InstanceInfo = Item->GetInstanceInfo();
+		TSharedPtr<ISessionInstanceInfo> InstanceInfo = Item->GetInstanceInfo();
 
 		if (InstanceInfo.IsValid() && !InstanceInfo->IsAuthorized())
 		{
@@ -177,7 +177,7 @@ private:
 	/** Callback for getting the text in the 'Device' column. */
 	FText HandleDeviceColumnText() const
 	{
-		ISessionInstanceInfoPtr InstanceInfo = Item->GetInstanceInfo();
+		TSharedPtr<ISessionInstanceInfo> InstanceInfo = Item->GetInstanceInfo();
 
 		if (InstanceInfo.IsValid())
 		{
@@ -190,7 +190,7 @@ private:
 	/** Callback for getting the border color for this row. */
 	FSlateColor HandleInstanceBorderBackgroundColor() const
 	{
-		ISessionInstanceInfoPtr InstanceInfo = Item->GetInstanceInfo();
+		TSharedPtr<ISessionInstanceInfo> InstanceInfo = Item->GetInstanceInfo();
 
 		if (InstanceInfo.IsValid())
 		{
@@ -203,7 +203,7 @@ private:
 	/** Callback for getting the border brush for this row. */
 	const FSlateBrush* HandleInstanceBorderBrush() const
 	{
-		ISessionInstanceInfoPtr InstanceInfo = Item->GetInstanceInfo();
+		TSharedPtr<ISessionInstanceInfo> InstanceInfo = Item->GetInstanceInfo();
 
 		if (InstanceInfo.IsValid())
 		{
@@ -221,7 +221,7 @@ private:
 	/** Callback for getting the type of the session instance. */
 	FText HandleInstanceTypeText() const
 	{
-		ISessionInstanceInfoPtr InstanceInfo = Item->GetInstanceInfo();
+		TSharedPtr<ISessionInstanceInfo> InstanceInfo = Item->GetInstanceInfo();
 
 		if (InstanceInfo.IsValid())
 		{
@@ -234,7 +234,7 @@ private:
 	/** Callback for getting the instance's current level. */
 	FText HandleLevelColumnText() const
 	{
-		ISessionInstanceInfoPtr InstanceInfo = Item->GetInstanceInfo();
+		TSharedPtr<ISessionInstanceInfo> InstanceInfo = Item->GetInstanceInfo();
 
 		if (InstanceInfo.IsValid())
 		{
@@ -247,7 +247,7 @@ private:
 	/** Callback for getting the image of the Status icon. */
 	const FSlateBrush* HandleStatusImage() const
 	{
-		ISessionInstanceInfoPtr InstanceInfo = Item->GetInstanceInfo();
+		TSharedPtr<ISessionInstanceInfo> InstanceInfo = Item->GetInstanceInfo();
 
 		if (InstanceInfo.IsValid())
 		{
@@ -265,7 +265,7 @@ private:
 	/** Callback for getting the foreground text color. */
 	FSlateColor HandleTextColorAndOpacity() const
 	{
-		ISessionInstanceInfoPtr InstanceInfo = Item->GetInstanceInfo();
+		TSharedPtr<ISessionInstanceInfo> InstanceInfo = Item->GetInstanceInfo();
 
 		if (InstanceInfo.IsValid())
 		{
diff --git a/Engine/Source/Developer/SessionFrontend/Private/Widgets/Browser/SSessionBrowserTreeSessionRow.h b/Engine/Source/Developer/SessionFrontend/Private/Widgets/Browser/SSessionBrowserTreeSessionRow.h
--- a/Engine/Source/Developer/SessionFrontend/Private/Widgets/Browser/SSessionBrowserTreeSessionRow.h
+++ b/Engine/Source/Developer/SessionFrontend/Private/Widgets/Browser/SSessionBrowserTreeSessionRow.h
@@ -73,7 +73,7 @@ public:
 
 protected:
 
-	FText SanitizeSessionName(const ISessionInfoPtr& SessionInfo) const
+	FText SanitizeSessionName(const TSharedPtr<ISessionInfo>& SessionInfo) const
 	{
 		const FString& SessionName = SessionInfo->GetSessionName();
 		const FString& SessionOwner = SessionInfo->GetSessionOwner();
@@ -90,12 +90,12 @@ protected:
 		}
 
 		// generate name for a standalone session
-		TArray<ISessionInstanceInfoPtr> Instances;
+		TArray<TSharedPtr<ISessionInstanceInfo>> Instances;
 		SessionInfo->GetInstances(Instances);
 
 		if (Instances.Num() > 0)
 		{
-			const ISessionInstanceInfoPtr& FirstInstance = Instances[0];
+			const TSharedPtr<ISessionInstanceInfo>& FirstInstance = Instances[0];
 
 			if ((Instances.Num() == 1) && FApp::IsThisInstance(FirstInstance->GetInstanceId()))
 			{
@@ -153,7 +153,7 @@ private:
 	{
 		FTextBuilder ToolTipTextBuilder;
 
-		ISessionInfoPtr SessionInfo = Item->GetSessionInfo();
+		TSharedPtr<ISessionInfo> SessionInfo = Item->GetSessionInfo();
 
 		if (SessionInfo.IsValid())
 		{
diff --git a/Engine/Source/Developer/SessionFrontend/Private/Widgets/Console/SSessionConsole.cpp b/Engine/Source/Developer/SessionFrontend/Private/Widgets/Console/SSessionConsole.cpp
--- a/Engine/Source/Developer/SessionFrontend/Private/Widgets/Console/SSessionConsole.cpp
+++ b/Engine/Source/Developer/SessionFrontend/Private/Widgets/Console/SSessionConsole.cpp
@@ -25,7 +25,7 @@ SSessionConsole::~SSessionConsole()
  *****************************************************************************/
 
 BEGIN_SLATE_FUNCTION_BUILD_OPTIMIZATION
-void SSessionConsole::Construct(const FArguments& InArgs, ISessionManagerRef InSessionManager)
+void SSessionConsole::Construct(const FArguments& InArgs, TSharedRef<ISessionManager> InSessionManager)
 {
 	SessionManager = InSessionManager;
 	ShouldScrollToLast = true;
@@ -81,7 +81,7 @@ void SSessionConsole::Construct(const FArguments& InArgs, ISessionManagerRef InS
 									.BorderImage(FEditorStyle::GetBrush("ToolPanel.GroupBorder"))
 									.Padding(0.0f)
 									[
-										SAssignNew(LogListView, SListView<FSessionLogMessagePtr>)
+										SAssignNew(LogListView, SListView<TSharedPtr<FSessionLogMessage>>)
 											.ItemHeight(24.0f)
 											.ListItemsSource(&LogMessages)
 											.SelectionMode(ESelectionMode::Multi)
@@ -194,7 +194,7 @@ void SSessionConsole::ClearLog()
 
 void SSessionConsole::CopyLog()
 {
-	TArray<FSessionLogMessagePtr> SelectedItems = LogListView->GetSelectedItems();
+	TArray<TSharedPtr<FSessionLogMessage>> SelectedItems = LogListView->GetSelectedItems();
 
 	if (SelectedItems.Num() == 0)
 	{
@@ -224,7 +224,7 @@ void SSessionConsole::ReloadLog(bool FullyReload)
 
 		for (const auto& Instance : SelectedInstances)
 		{
-			const TArray<FSessionLogMessagePtr>& InstanceLog = Instance->GetLog();
+			const TArray<TSharedPtr<FSessionLogMessage>>& InstanceLog = Instance->GetLog();
 
 			for (const auto& LogMessage : InstanceLog)
 			{
@@ -418,7 +418,7 @@ void SSessionConsole::HandleFilterChanged()
 }
 
 
-void SSessionConsole::HandleLogListItemScrolledIntoView(FSessionLogMessagePtr Item, const TSharedPtr<ITableRow>& TableRow)
+void SSessionConsole::HandleLogListItemScrolledIntoView(TSharedPtr<FSessionLogMessage> Item, const TSharedPtr<ITableRow>& TableRow)
 {
 	if (LogMessages.Num() > 0)
 	{
@@ -431,7 +431,7 @@ void SSessionConsole::HandleLogListItemScrolledIntoView(FSessionLogMessagePtr It
 }
 
 
-TSharedRef<ITableRow> SSessionConsole::HandleLogListGenerateRow(FSessionLogMessagePtr Message, const TSharedRef<STableViewBase>& OwnerTable)
+TSharedRef<ITableRow> SSessionConsole::HandleLogListGenerateRow(TSharedPtr<FSessionLogMessage> Message, const TSharedRef<STableViewBase>& OwnerTable)
 {
 	return SNew(SSessionConsoleLogTableRow, OwnerTable)
 		.HighlightText(this, &SSessionConsole::HandleLogListGetHighlightText)
@@ -481,7 +481,7 @@ void SSessionConsole::HandleSessionManagerInstanceSelectionChanged(const TShared
 }
 
 
-void SSessionConsole::HandleSessionManagerLogReceived(const ISessionInfoRef& Session, const ISessionInstanceInfoRef& Instance, const FSessionLogMessageRef& Message)
+void SSessionConsole::HandleSessionManagerLogReceived(const TSharedRef<ISessionInfo>& Session, const TSharedRef<ISessionInstanceInfo>& Instance, const TSharedRef<FSessionLogMessage>& Message)
 {
 	if (!SessionManager->IsInstanceSelected(Instance) || !FilterBar->FilterLogMessage(Message))
 	{
@@ -500,7 +500,7 @@ void SSessionConsole::HandleSessionManagerLogReceived(const ISessionInfoRef& Ses
 }
 
 
-void SSessionConsole::HandleSessionManagerSelectedSessionChanged(const ISessionInfoPtr& SelectedSession)
+void SSessionConsole::HandleSessionManagerSelectedSessionChanged(const TSharedPtr<ISessionInfo>& SelectedSession)
 {
 	ReloadLog(true);
 }
diff --git a/Engine/Source/Developer/SessionFrontend/Private/Widgets/Console/SSessionConsole.h b/Engine/Source/Developer/SessionFrontend/Private/Widgets/Console/SSessionConsole.h
--- a/Engine/Source/Developer/SessionFrontend/Private/Widgets/Console/SSessionConsole.h
+++ b/Engine/Source/Developer/SessionFrontend/Private/Widgets/Console/SSessionConsole.h
@@ -30,7 +30,7 @@ public:
 	 * @param InArgs The declaration data for this widget.
 	 * @param InSessionManager The session manager to use.
 	 */
-	void Construct(const FArguments& InArgs, ISessionManagerRef InSessionManager);
+	void Construct(const FArguments& InArgs, TSharedRef<ISessionManager> InSessionManager);
 
 protected:
 
@@ -109,16 +109,16 @@ private:
 	void HandleFilterChanged();
 
 	/** Callback for scrolling a log item into view. */
-	void HandleLogListItemScrolledIntoView(FSessionLogMessagePtr Item, const TSharedPtr<ITableRow>& TableRow);
+	void HandleLogListItemScrolledIntoView(TSharedPtr<FSessionLogMessage> Item, const TSharedPtr<ITableRow>& TableRow);
 
 	/** Callback for generating a row widget for the log list view. */
-	TSharedRef<ITableRow> HandleLogListGenerateRow(FSessionLogMessagePtr Message, const TSharedRef<STableViewBase>& OwnerTable);
+	TSharedRef<ITableRow> HandleLogListGenerateRow(TSharedPtr<FSessionLogMessage> Message, const TSharedRef<STableViewBase>& OwnerTable);
 
 	/** Callback for getting the highlight string for log messages. */
 	FText HandleLogListGetHighlightText() const;
 
 	/** Callback for selecting log messages. */
-	void HandleLogListSelectionChanged(FSessionLogMessagePtr InItem, ESelectInfo::Type SelectInfo);
+	void HandleLogListSelectionChanged(TSharedPtr<FSessionLogMessage> InItem, ESelectInfo::Type SelectInfo);
 
 	/** Callback for getting the enabled state of the console box. */
 	bool HandleMainContentIsEnabled() const;
@@ -130,15 +130,15 @@ private:
 	void HandleSessionManagerInstanceSelectionChanged(const TSharedPtr<ISessionInstanceInfo>& Instance, bool Selected);
 
 	/** Callback for received log entries. */
-	void HandleSessionManagerLogReceived(const ISessionInfoRef& Session, const ISessionInstanceInfoRef& Instance, const FSessionLogMessageRef& Message);
+	void HandleSessionManagerLogReceived(const TSharedRef<ISessionInfo>& Session, const TSharedRef<ISessionInstanceInfo>& Instance, const TSharedRef<FSessionLogMessage>& Message);
 
 	/** Callback for changing the selected session. */
-	void HandleSessionManagerSelectedSessionChanged(const ISessionInfoPtr& SelectedSession);
+	void HandleSessionManagerSelectedSessionChanged(const TSharedPtr<ISessionInfo>& SelectedSession);
 
 private:
 
 	/** Holds an unfiltered list of available log messages. */
-	TArray<FSessionLogMessagePtr> AvailableLogs;
+	TArray<TSharedPtr<FSessionLogMessage>> AvailableLogs;
 
 	/** Holds the command bar. */
 	TSharedPtr<SSessionConsoleCommandBar> CommandBar;
@@ -156,13 +156,13 @@ private:
 	FString LastLogFileSaveDirectory;
 
 	/** Holds the log list view. */
- 	TSharedPtr<SListView<FSessionLogMessagePtr>> LogListView;
+ 	TSharedPtr<SListView<TSharedPtr<FSessionLogMessage>>> LogListView;
 
  	/** Holds the filtered list of log messages. */
- 	TArray<FSessionLogMessagePtr> LogMessages;
+ 	TArray<TSharedPtr<FSessionLogMessage>> LogMessages;
 
 	/** Holds the session manager. */
-	ISessionManagerPtr SessionManager;
+	TSharedPtr<ISessionManager> SessionManager;
 
 	/** Holds the shortcut window. */
 	TSharedPtr<SSessionConsoleShortcutWindow> ShortcutWindow;
diff --git a/Engine/Source/Developer/SessionFrontend/Private/Widgets/Console/SSessionConsoleFilterBar.cpp b/Engine/Source/Developer/SessionFrontend/Private/Widgets/Console/SSessionConsoleFilterBar.cpp
--- a/Engine/Source/Developer/SessionFrontend/Private/Widgets/Console/SSessionConsoleFilterBar.cpp
+++ b/Engine/Source/Developer/SessionFrontend/Private/Widgets/Console/SSessionConsoleFilterBar.cpp
@@ -100,7 +100,7 @@ void SSessionConsoleFilterBar::Construct(const FArguments& InArgs)
 END_SLATE_FUNCTION_BUILD_OPTIMIZATION
 
 
-bool SSessionConsoleFilterBar::FilterLogMessage(const FSessionLogMessageRef& LogMessage)
+bool SSessionConsoleFilterBar::FilterLogMessage(const TSharedRef<FSessionLogMessage>& LogMessage)
 {
 	// create or update category counter
 	int32& CategoryCounter = CategoryCounters.FindOrAdd(LogMessage->Category);
diff --git a/Engine/Source/Developer/SessionFrontend/Private/Widgets/Console/SSessionConsoleFilterBar.h b/Engine/Source/Developer/SessionFrontend/Private/Widgets/Console/SSessionConsoleFilterBar.h
--- a/Engine/Source/Developer/SessionFrontend/Private/Widgets/Console/SSessionConsoleFilterBar.h
+++ b/Engine/Source/Developer/SessionFrontend/Private/Widgets/Console/SSessionConsoleFilterBar.h
@@ -33,7 +33,7 @@ public:
 	 * @param LogMessage The log message to filter.
 	 * @return true if the log message passed the filter, false otherwise.
 	 */
-	bool FilterLogMessage(const FSessionLogMessageRef& LogMessage);
+	bool FilterLogMessage(const TSharedRef<FSessionLogMessage>& LogMessage);
 
 	/**
 	 * Gets the current filter string.
diff --git a/Engine/Source/Developer/SessionFrontend/Private/Widgets/Console/SSessionConsoleLogTableRow.h b/Engine/Source/Developer/SessionFrontend/Private/Widgets/Console/SSessionConsoleLogTableRow.h
--- a/Engine/Source/Developer/SessionFrontend/Private/Widgets/Console/SSessionConsoleLogTableRow.h
+++ b/Engine/Source/Developer/SessionFrontend/Private/Widgets/Console/SSessionConsoleLogTableRow.h
@@ -7,13 +7,13 @@
  * Implements a row widget for the session console log.
  */
 class SSessionConsoleLogTableRow
-	: public SMultiColumnTableRow<FSessionLogMessagePtr>
+	: public SMultiColumnTableRow<TSharedPtr<FSessionLogMessage>>
 {
 public:
 
 	SLATE_BEGIN_ARGS(SSessionConsoleLogTableRow) { }
 		SLATE_ATTRIBUTE(FText, HighlightText)
-		SLATE_ARGUMENT(FSessionLogMessagePtr, LogMessage)
+		SLATE_ARGUMENT(TSharedPtr<FSessionLogMessage>, LogMessage)
 	SLATE_END_ARGS()
 
 public:
@@ -28,7 +28,7 @@ public:
 		HighlightText = InArgs._HighlightText;
 		LogMessage = InArgs._LogMessage;
 
-		SMultiColumnTableRow<FSessionLogMessagePtr>::Construct(FSuperRowType::FArguments(), InOwnerTableView);
+		SMultiColumnTableRow<TSharedPtr<FSessionLogMessage>>::Construct(FSuperRowType::FArguments(), InOwnerTableView);
 	}
 
 public:
@@ -148,5 +148,5 @@ private:
 	TAttribute<FText> HighlightText;
 
 	/** Holds a reference to the log message that is displayed in this row. */
-	FSessionLogMessagePtr LogMessage;
+	TSharedPtr<FSessionLogMessage> LogMessage;
 };
diff --git a/Engine/Source/Developer/SessionFrontend/Private/Widgets/SSessionFrontend.h b/Engine/Source/Developer/SessionFrontend/Private/Widgets/SSessionFrontend.h
--- a/Engine/Source/Developer/SessionFrontend/Private/Widgets/SSessionFrontend.h
+++ b/Engine/Source/Developer/SessionFrontend/Private/Widgets/SSessionFrontend.h
@@ -60,7 +60,7 @@ private:
 	TSharedPtr<SButton> NewSessionButton;
 
 	/** Holds a pointer to the session manager. */
-	ISessionManagerPtr SessionManager;
+	TSharedPtr<ISessionManager> SessionManager;
 
 	/** Holds a pointer to the session manager. */
 	IScreenShotManagerPtr ScreenShotManager;
diff --git a/Engine/Source/Developer/SlateReflector/Private/WidgetSnapshotService.cpp b/Engine/Source/Developer/SlateReflector/Private/WidgetSnapshotService.cpp
--- a/Engine/Source/Developer/SlateReflector/Private/WidgetSnapshotService.cpp
+++ b/Engine/Source/Developer/SlateReflector/Private/WidgetSnapshotService.cpp
@@ -10,7 +10,7 @@ FWidgetSnapshotService::FWidgetSnapshotService()
 {
 	if (FPlatformMisc::SupportsMessaging() && FPlatformProcess::SupportsMultithreading())
 	{
-		IMessageBusPtr MessageBusPtr = IMessagingModule::Get().GetDefaultBus();
+		TSharedPtr<IMessageBus, ESPMode::ThreadSafe> MessageBusPtr = IMessagingModule::Get().GetDefaultBus();
 
 		MessageEndpoint = FMessageEndpoint::Builder("FWidgetSnapshotService", MessageBusPtr.ToSharedRef())
 			.ReceivingOnThread(ENamedThreads::GameThread)
diff --git a/Engine/Source/Developer/SlateReflector/Private/Widgets/SWidgetReflector.cpp b/Engine/Source/Developer/SlateReflector/Private/Widgets/SWidgetReflector.cpp
--- a/Engine/Source/Developer/SlateReflector/Private/Widgets/SWidgetReflector.cpp
+++ b/Engine/Source/Developer/SlateReflector/Private/Widgets/SWidgetReflector.cpp
@@ -391,7 +391,7 @@ void SWidgetReflector::Construct( const FArguments& InArgs )
 
 #if SLATE_REFLECTOR_HAS_SESSION_SERVICES
 	{
-		ISessionManagerRef SessionManager = FModuleManager::LoadModuleChecked<ISessionServicesModule>("SessionServices").GetSessionManager();
+		TSharedRef<ISessionManager> SessionManager = FModuleManager::LoadModuleChecked<ISessionServicesModule>("SessionServices").GetSessionManager();
 		SessionManager->OnSessionsUpdated().AddSP(this, &SWidgetReflector::OnAvailableSnapshotTargetsChanged);
 	}
 #endif // SLATE_REFLECTOR_HAS_SESSION_SERVICES
@@ -1581,7 +1581,7 @@ void SWidgetReflector::UpdateAvailableSnapshotTargets()
 
 #if SLATE_REFLECTOR_HAS_SESSION_SERVICES
 	{
-		ISessionManagerRef SessionManager = FModuleManager::LoadModuleChecked<ISessionServicesModule>("SessionServices").GetSessionManager();
+		TSharedRef<ISessionManager> SessionManager = FModuleManager::LoadModuleChecked<ISessionServicesModule>("SessionServices").GetSessionManager();
 
 		TArray<TSharedPtr<ISessionInfo>> AvailableSessions;
 		SessionManager->GetSessions(AvailableSessions);
diff --git a/Engine/Source/Developer/TargetDeviceServices/Private/TargetDeviceProxy.cpp b/Engine/Source/Developer/TargetDeviceServices/Private/TargetDeviceProxy.cpp
--- a/Engine/Source/Developer/TargetDeviceServices/Private/TargetDeviceProxy.cpp
+++ b/Engine/Source/Developer/TargetDeviceServices/Private/TargetDeviceProxy.cpp
@@ -212,7 +212,7 @@ bool FTargetDeviceProxy::DeployApp(FName InVariant, const TMap<FString, FString>
 {
 	for (TMap<FString, FString>::TConstIterator It(Files); It; ++It)
 	{
-		IMessageAttachmentRef FileAttachment = MakeShareable(new FFileMessageAttachment(It.Key()));
+		TSharedRef<IMessageAttachment, ESPMode::ThreadSafe> FileAttachment = MakeShareable(new FFileMessageAttachment(It.Key()));
 		FString SourcePath = It.Key();
 
 		MessageEndpoint->Send(new FTargetDeviceServiceDeployFile(It.Value(), TransactionId), FileAttachment, MessageAddress);
diff --git a/Engine/Source/Developer/TargetDeviceServices/Private/TargetDeviceService.cpp b/Engine/Source/Developer/TargetDeviceServices/Private/TargetDeviceService.cpp
--- a/Engine/Source/Developer/TargetDeviceServices/Private/TargetDeviceService.cpp
+++ b/Engine/Source/Developer/TargetDeviceServices/Private/TargetDeviceService.cpp
@@ -22,7 +22,7 @@ struct FVariantSortCallback
 /* FTargetDeviceService structors
  *****************************************************************************/
 
-FTargetDeviceService::FTargetDeviceService(const FString& InDeviceName, const IMessageBusRef& InMessageBus)
+FTargetDeviceService::FTargetDeviceService(const FString& InDeviceName, const TSharedRef<IMessageBus, ESPMode::ThreadSafe>& InMessageBus)
 	: DeviceName(InDeviceName)
 	, Running(false)
 	, Shared(false)
@@ -317,7 +317,7 @@ void FTargetDeviceService::HandleDeployFileMessage(const FTargetDeviceServiceDep
 		return;
 	}
 
-	IMessageAttachmentPtr Attachment = Context->GetAttachment();
+	TSharedPtr<IMessageAttachment, ESPMode::ThreadSafe> Attachment = Context->GetAttachment();
 
 	if (Attachment.IsValid())
 	{
diff --git a/Engine/Source/Developer/TargetDeviceServices/Private/TargetDeviceService.h b/Engine/Source/Developer/TargetDeviceServices/Private/TargetDeviceService.h
--- a/Engine/Source/Developer/TargetDeviceServices/Private/TargetDeviceService.h
+++ b/Engine/Source/Developer/TargetDeviceServices/Private/TargetDeviceService.h
@@ -16,7 +16,7 @@ public:
 	 * @param InDeviceName The name of the device to expose.
 	 * @param InMessageBus The message bus to listen on for clients.
 	 */
-	FTargetDeviceService(const FString& InDeviceName, const IMessageBusRef& InMessageBus);
+	FTargetDeviceService(const FString& InDeviceName, const TSharedRef<IMessageBus, ESPMode::ThreadSafe>& InMessageBus);
 
 	/** Destructor. */
 	~FTargetDeviceService();
diff --git a/Engine/Source/Developer/TargetDeviceServices/Private/TargetDeviceServiceManager.cpp b/Engine/Source/Developer/TargetDeviceServices/Private/TargetDeviceServiceManager.cpp
--- a/Engine/Source/Developer/TargetDeviceServices/Private/TargetDeviceServiceManager.cpp
+++ b/Engine/Source/Developer/TargetDeviceServices/Private/TargetDeviceServiceManager.cpp
@@ -8,7 +8,7 @@
 
 FTargetDeviceServiceManager::FTargetDeviceServiceManager()
 {
-	IMessageBusPtr MessageBus = IMessagingModule::Get().GetDefaultBus();
+	TSharedPtr<IMessageBus, ESPMode::ThreadSafe> MessageBus = IMessagingModule::Get().GetDefaultBus();
 
 	if (MessageBus.IsValid())
 	{
@@ -30,7 +30,7 @@ FTargetDeviceServiceManager::~FTargetDeviceServiceManager()
 		SaveSettings();
 	}
 
-	IMessageBusPtr MessageBus = MessageBusPtr.Pin();
+	auto MessageBus = MessageBusPtr.Pin();
 
 	if (MessageBus.IsValid())
 	{
@@ -81,7 +81,7 @@ void FTargetDeviceServiceManager::RemoveStartupService(const FString& DeviceName
 
 ITargetDeviceServicePtr FTargetDeviceServiceManager::AddService(const FString& DeviceName)
 {
-	IMessageBusPtr MessageBus = MessageBusPtr.Pin();
+	auto MessageBus = MessageBusPtr.Pin();
 
 	if (!MessageBus.IsValid())
 	{
@@ -113,14 +113,15 @@ ITargetDeviceServicePtr FTargetDeviceServiceManager::AddService(const FString& D
 bool FTargetDeviceServiceManager::AddTargetDevice(ITargetDevicePtr InDevice)
 {
 	FScopeLock Lock(&CriticalSection);
-	IMessageBusPtr MessageBus = MessageBusPtr.Pin();
+	
+	auto MessageBus = MessageBusPtr.Pin();
+	
 	if (!MessageBus.IsValid())
 	{
 		return false;
 	}
 
 	const FString& DeviceName = InDevice->GetName();
-
 	ITargetDeviceServicePtr DeviceService = AddService(DeviceName);
 
 	if (DeviceService.IsValid())
diff --git a/Engine/Source/Developer/TargetDeviceServices/Private/TargetDeviceServiceManager.h b/Engine/Source/Developer/TargetDeviceServices/Private/TargetDeviceServiceManager.h
--- a/Engine/Source/Developer/TargetDeviceServices/Private/TargetDeviceServiceManager.h
+++ b/Engine/Source/Developer/TargetDeviceServices/Private/TargetDeviceServiceManager.h
@@ -120,7 +120,7 @@ private:
 	TMap<FString, ITargetDeviceServicePtr> DeviceServices;
 
 	/** Holds a weak pointer to the message bus. */
-	IMessageBusWeakPtr MessageBusPtr;
+	TWeakPtr<IMessageBus, ESPMode::ThreadSafe> MessageBusPtr;
 
 	/** Holds the collection of identifiers for devices that start automatically (shared/unshared). */
 	TMap<FString, bool> StartupServices;
diff --git a/Engine/Source/Editor/DetailCustomizations/DetailCustomizations.Build.cs b/Engine/Source/Editor/DetailCustomizations/DetailCustomizations.Build.cs
--- a/Engine/Source/Editor/DetailCustomizations/DetailCustomizations.Build.cs
+++ b/Engine/Source/Editor/DetailCustomizations/DetailCustomizations.Build.cs
@@ -24,6 +24,7 @@ public class DetailCustomizations : ModuleRules
 				"EditorWidgets",
 				"KismetWidgets",
 				"MovieSceneCapture",
+				"MovieSceneTools",
 				"SharedSettingsWidgets",
                 "ContentBrowser",
 				"BlueprintGraph",
diff --git a/Engine/Source/Editor/DetailCustomizations/Private/DetailCustomizations.cpp b/Engine/Source/Editor/DetailCustomizations/Private/DetailCustomizations.cpp
--- a/Engine/Source/Editor/DetailCustomizations/Private/DetailCustomizations.cpp
+++ b/Engine/Source/Editor/DetailCustomizations/Private/DetailCustomizations.cpp
@@ -104,6 +104,9 @@
 #include "CaptureTypeCustomization.h"
 #include "RenderPassesCustomization.h"
 #include "MovieSceneCaptureCustomization.h"
+#include "MovieSceneEvalOptionsCustomization.h"
+#include "MovieSceneEventParametersCustomization.h"
+#include "MovieSceneCurveInterfaceKeyEditStructCustomization.h"
 #include "LevelSequenceBurnInOptionsCustomization.h"
 #include "TextCustomization.h"
 #include "AnimTrailNodeDetails.h"
@@ -232,6 +235,9 @@ void FDetailCustomizationsModule::RegisterPropertyTypeCustomizations()
 	RegisterCustomPropertyTypeLayout("CameraFilmbackSettings", FOnGetPropertyTypeCustomizationInstance::CreateStatic(&FCameraFilmbackSettingsCustomization::MakeInstance));
 	RegisterCustomPropertyTypeLayout("CameraLensSettings", FOnGetPropertyTypeCustomizationInstance::CreateStatic(&FCameraLensSettingsCustomization::MakeInstance));
 	RegisterCustomPropertyTypeLayout("CameraFocusSettings", FOnGetPropertyTypeCustomizationInstance::CreateStatic(&FCameraFocusSettingsCustomization::MakeInstance));
+	RegisterCustomPropertyTypeLayout("MovieSceneTrackEvalOptions", FOnGetPropertyTypeCustomizationInstance::CreateStatic(&FMovieSceneTrackEvalOptionsCustomization::MakeInstance));
+	RegisterCustomPropertyTypeLayout("MovieSceneSectionEvalOptions", FOnGetPropertyTypeCustomizationInstance::CreateStatic(&FMovieSceneSectionEvalOptionsCustomization::MakeInstance));
+	RegisterCustomPropertyTypeLayout("MovieSceneEventParameters", FOnGetPropertyTypeCustomizationInstance::CreateStatic(&FMovieSceneEventParametersCustomization::MakeInstance));
 	RegisterCustomPropertyTypeLayout("LevelSequenceBurnInOptions", FOnGetPropertyTypeCustomizationInstance::CreateStatic(&FLevelSequenceBurnInOptionsCustomization::MakeInstance));
 	RegisterCustomPropertyTypeLayout("LevelSequenceBurnInInitSettings", FOnGetPropertyTypeCustomizationInstance::CreateStatic(&FLevelSequenceBurnInInitSettingsCustomization::MakeInstance));
 }
@@ -325,6 +331,7 @@ void FDetailCustomizationsModule::RegisterObjectCustomizations()
 	RegisterCustomClassLayout("ConfigHierarchyPropertyView", FOnGetDetailCustomizationInstance::CreateStatic(&FConfigPropertyHelperDetails::MakeInstance));
 
 	RegisterCustomClassLayout("MovieSceneCapture", FOnGetDetailCustomizationInstance::CreateStatic(&FMovieSceneCaptureCustomization::MakeInstance));
+	RegisterCustomClassLayout("MovieSceneCurveInterfaceKeyEditStruct", FOnGetDetailCustomizationInstance::CreateStatic(&FMovieSceneCurveInterfaceKeyEditStructCustomization::MakeInstance));
 
 	RegisterCustomClassLayout("AnalyticsPrivacySettings", FOnGetDetailCustomizationInstance::CreateStatic(&FImportantToggleSettingCustomization::MakeInstance));
 	RegisterCustomClassLayout("EndUserSettings", FOnGetDetailCustomizationInstance::CreateStatic(&FImportantToggleSettingCustomization::MakeInstance));
diff --git a/Engine/Source/Editor/DetailCustomizations/Private/MovieSceneCurveInterfaceKeyEditStructCustomization.cpp b/Engine/Source/Editor/DetailCustomizations/Private/MovieSceneCurveInterfaceKeyEditStructCustomization.cpp
new file mode 100644
--- /dev/null
+++ b/Engine/Source/Editor/DetailCustomizations/Private/MovieSceneCurveInterfaceKeyEditStructCustomization.cpp
@@ -0,0 +1,40 @@
+// Copyright 1998-2016 Epic Games, Inc. All Rights Reserved.
+
+#include "DetailCustomizationsPrivatePCH.h"
+#include "MovieSceneCurveInterfaceKeyEditStructCustomization.h"
+
+#include "GenericKeyArea.h"
+
+#define LOCTEXT_NAMESPACE "MovieSceneCurveInterfaceKeyEditStructCustomization"
+
+TSharedRef<IDetailCustomization> FMovieSceneCurveInterfaceKeyEditStructCustomization::MakeInstance()
+{
+	return MakeShared<FMovieSceneCurveInterfaceKeyEditStructCustomization>();
+}
+
+void FMovieSceneCurveInterfaceKeyEditStructCustomization::CustomizeDetails(IDetailLayoutBuilder& DetailBuilder)
+{
+	TArray<TSharedPtr<FStructOnScope>> Structs;
+	DetailBuilder.GetStructsBeingCustomized(Structs);
+
+	if (Structs.Num() != 1)
+	{
+		return;
+	}
+
+	const UStruct* StructType = Structs[0]->GetStruct();
+	if (!StructType || StructType != FMovieSceneCurveInterfaceKeyEditStruct::StaticStruct())
+	{
+		return;
+	}
+
+	auto* StructValue = (FMovieSceneCurveInterfaceKeyEditStruct*)Structs[0]->GetStructMemory();
+	if (!StructValue)
+	{
+		return;
+	}
+
+	StructValue->EditInterface->Extend(StructValue->KeyHandle, DetailBuilder);
+}
+
+#undef LOCTEXT_NAMESPACE
\ No newline at end of file
diff --git a/Engine/Source/Editor/DetailCustomizations/Private/MovieSceneCurveInterfaceKeyEditStructCustomization.h b/Engine/Source/Editor/DetailCustomizations/Private/MovieSceneCurveInterfaceKeyEditStructCustomization.h
new file mode 100644
--- /dev/null
+++ b/Engine/Source/Editor/DetailCustomizations/Private/MovieSceneCurveInterfaceKeyEditStructCustomization.h
@@ -0,0 +1,13 @@
+// Copyright 1998-2016 Epic Games, Inc. All Rights Reserved.
+
+#pragma once
+
+class IDetailsView;
+
+class FMovieSceneCurveInterfaceKeyEditStructCustomization : public IDetailCustomization
+{
+public:
+	static TSharedRef<IDetailCustomization> MakeInstance();
+
+	virtual void CustomizeDetails(IDetailLayoutBuilder& DetailBuilder) override;
+};
\ No newline at end of file
diff --git a/Engine/Source/Editor/DetailCustomizations/Private/MovieSceneEvalOptionsCustomization.cpp b/Engine/Source/Editor/DetailCustomizations/Private/MovieSceneEvalOptionsCustomization.cpp
new file mode 100644
--- /dev/null
+++ b/Engine/Source/Editor/DetailCustomizations/Private/MovieSceneEvalOptionsCustomization.cpp
@@ -0,0 +1,61 @@
+// Copyright 1998-2016 Epic Games, Inc. All Rights Reserved.
+
+#include "DetailCustomizationsPrivatePCH.h"
+#include "MovieSceneEvalOptionsCustomization.h"
+
+#include "MovieSceneTrack.h"
+#include "MovieSceneSection.h"
+
+TSharedRef<IPropertyTypeCustomization> FMovieSceneTrackEvalOptionsCustomization::MakeInstance()
+{
+	return MakeShared<FMovieSceneTrackEvalOptionsCustomization>();
+}
+
+void FMovieSceneTrackEvalOptionsCustomization::CustomizeHeader(TSharedRef<IPropertyHandle> PropertyHandle, FDetailWidgetRow& HeaderRow, IPropertyTypeCustomizationUtils& CustomizationUtils)
+{
+}
+
+void FMovieSceneTrackEvalOptionsCustomization::CustomizeChildren(TSharedRef<IPropertyHandle> PropertyHandle, IDetailChildrenBuilder& ChildBuilder, IPropertyTypeCustomizationUtils& CustomizationUtils)
+{
+	TArray<void*> RawData;
+	PropertyHandle->AccessRawData(RawData);
+
+	bool bCanEvaluateNearestSection = !RawData.ContainsByPredicate(
+		[](void* Ptr){
+			return !static_cast<FMovieSceneTrackEvalOptions*>(Ptr)->bCanEvaluateNearestSection;
+		}
+	);
+
+	TSharedPtr<IPropertyHandle> bEvaluateNearestSectionHandle = PropertyHandle->GetChildHandle(GET_MEMBER_NAME_CHECKED(FMovieSceneTrackEvalOptions, bEvaluateNearestSection));
+	if (bCanEvaluateNearestSection && bEvaluateNearestSectionHandle.IsValid())
+	{
+		ChildBuilder.AddChildProperty(bEvaluateNearestSectionHandle.ToSharedRef());
+	}
+}
+
+TSharedRef<IPropertyTypeCustomization> FMovieSceneSectionEvalOptionsCustomization::MakeInstance()
+{
+	return MakeShared<FMovieSceneSectionEvalOptionsCustomization>();
+}
+
+void FMovieSceneSectionEvalOptionsCustomization::CustomizeHeader(TSharedRef<IPropertyHandle> PropertyHandle, FDetailWidgetRow& HeaderRow, IPropertyTypeCustomizationUtils& CustomizationUtils)
+{
+}
+
+void FMovieSceneSectionEvalOptionsCustomization::CustomizeChildren(TSharedRef<IPropertyHandle> PropertyHandle, IDetailChildrenBuilder& ChildBuilder, IPropertyTypeCustomizationUtils& CustomizationUtils)
+{
+	TArray<void*> RawData;
+	PropertyHandle->AccessRawData(RawData);
+
+	bool bCanEditCompletionMode = !RawData.ContainsByPredicate(
+		[](void* Ptr){
+			return !static_cast<FMovieSceneSectionEvalOptions*>(Ptr)->bCanEditCompletionMode;
+		}
+	);
+
+	TSharedPtr<IPropertyHandle> CompletionModeHandle = PropertyHandle->GetChildHandle(GET_MEMBER_NAME_CHECKED(FMovieSceneSectionEvalOptions, CompletionMode));
+	if (bCanEditCompletionMode && CompletionModeHandle.IsValid())
+	{
+		ChildBuilder.AddChildProperty(CompletionModeHandle.ToSharedRef());
+	}
+}
\ No newline at end of file
diff --git a/Engine/Source/Editor/DetailCustomizations/Private/MovieSceneEventParametersCustomization.cpp b/Engine/Source/Editor/DetailCustomizations/Private/MovieSceneEventParametersCustomization.cpp
new file mode 100644
--- /dev/null
+++ b/Engine/Source/Editor/DetailCustomizations/Private/MovieSceneEventParametersCustomization.cpp
@@ -0,0 +1,123 @@
+// Copyright 1998-2016 Epic Games, Inc. All Rights Reserved.
+
+#include "DetailCustomizationsPrivatePCH.h"
+#include "MovieSceneEventParametersCustomization.h"
+
+#include "MovieSceneEventSection.h"
+
+#include "PropertyCustomizationHelpers.h"
+#include "ContentBrowserModule.h"
+#include "IStructureDetailsView.h"
+#include "IPropertyUtilities.h"
+
+#define LOCTEXT_NAMESPACE "MovieSceneEventParameters"
+
+TSharedRef<IPropertyTypeCustomization> FMovieSceneEventParametersCustomization::MakeInstance()
+{
+	return MakeShared<FMovieSceneEventParametersCustomization>();
+}
+
+void FMovieSceneEventParametersCustomization::CustomizeHeader(TSharedRef<IPropertyHandle> InPropertyHandle, FDetailWidgetRow& HeaderRow, IPropertyTypeCustomizationUtils& CustomizationUtils)
+{
+	PropertyUtilities = CustomizationUtils.GetPropertyUtilities();
+}
+
+void FMovieSceneEventParametersCustomization::CustomizeChildren(TSharedRef<IPropertyHandle> InPropertyHandle, IDetailChildrenBuilder& ChildBuilder, IPropertyTypeCustomizationUtils& CustomizationUtils)
+{
+	PropertyHandle = InPropertyHandle;
+
+	TArray<void*> RawData;
+	PropertyHandle->AccessRawData(RawData);
+	if (RawData.Num() != 1)
+	{
+		return;
+	}
+
+	EditStructData = MakeShared<FStructOnScope>(nullptr);
+	static_cast<FMovieSceneEventParameters*>(RawData[0])->GetInstance(*EditStructData);
+
+	ChildBuilder.AddChildContent(LOCTEXT("ParametersText", "Parameters"))
+	.NameContent()
+	[
+		SNew(STextBlock)
+		.Font(CustomizationUtils.GetRegularFont())
+		.Text(LOCTEXT("ParameterStructType", "Parameter Struct"))
+	]
+	.ValueContent()
+	.VAlign(VAlign_Top)
+	.MaxDesiredWidth(TOptional<float>())
+	[
+		SNew(SObjectPropertyEntryBox)
+		.ObjectPath(EditStructData->GetStruct() ? EditStructData->GetStruct()->GetPathName() : FString())
+		.AllowedClass(UScriptStruct::StaticClass())
+		.OnObjectChanged(this, &FMovieSceneEventParametersCustomization::OnStructChanged)
+	];
+
+	if (EditStructData->GetStructMemory())
+	{
+		FSimpleDelegate OnEditStructContentsChangedDelegate = FSimpleDelegate::CreateSP(this, &FMovieSceneEventParametersCustomization::OnEditStructContentsChanged);
+
+		TArray<TSharedPtr<IPropertyHandle>> ExternalHandles = ChildBuilder.AddChildStructure(EditStructData.ToSharedRef());
+		for (TSharedPtr<IPropertyHandle> Handle : ExternalHandles)
+		{
+			Handle->SetOnPropertyValueChanged(OnEditStructContentsChangedDelegate);
+			Handle->SetOnChildPropertyValueChanged(OnEditStructContentsChangedDelegate);
+		}
+	}
+}
+
+void FMovieSceneEventParametersCustomization::OnStructChanged(const FAssetData& AssetData)
+{
+	UScriptStruct* NewStruct = nullptr;
+
+	if (AssetData.IsValid())
+	{
+		NewStruct = Cast<UScriptStruct>(AssetData.GetAsset());
+		if (!NewStruct)
+		{
+			// Don't clear the type if the user managed to just choose the wrong type of object
+			return;
+		}
+	}
+
+	TArray<void*> RawData;
+	PropertyHandle->AccessRawData(RawData);
+	for (void* Value : RawData)
+	{
+		static_cast<FMovieSceneEventParameters*>(Value)->Reassign(NewStruct);
+	}
+
+	FPropertyChangedEvent BubbleChangeEvent(PropertyHandle->GetProperty(), EPropertyChangeType::ValueSet, nullptr);
+	PropertyUtilities->NotifyFinishedChangingProperties(BubbleChangeEvent);
+	
+	PropertyUtilities->ForceRefresh();
+}
+
+void FMovieSceneEventParametersCustomization::OnEditStructContentsChanged()
+{
+	// @todo: call modify on the outer object if possible
+	UScriptStruct* Struct = Cast<UScriptStruct>((UStruct*)EditStructData->GetStruct());
+	if (!Struct)
+	{
+		return;
+	}
+
+	TArray<uint8> Bytes;
+	{
+		FMemoryWriter Writer(Bytes);
+		Struct->SerializeTaggedProperties(Writer, EditStructData->GetStructMemory(), Struct, nullptr);
+	}
+
+	TArray<void*> RawData;
+	PropertyHandle->AccessRawData(RawData);
+
+	for (void* Value : RawData)
+	{
+		static_cast<FMovieSceneEventParameters*>(Value)->OverwriteWith(Bytes);
+	}
+
+	FPropertyChangedEvent BubbleChangeEvent(PropertyHandle->GetProperty(), EPropertyChangeType::ValueSet, nullptr);
+	PropertyUtilities->NotifyFinishedChangingProperties(BubbleChangeEvent);
+}
+
+#undef LOCTEXT_NAMESPACE
\ No newline at end of file
diff --git a/Engine/Source/Editor/EditorStyle/Private/SlateEditorStyle.cpp b/Engine/Source/Editor/EditorStyle/Private/SlateEditorStyle.cpp
--- a/Engine/Source/Editor/EditorStyle/Private/SlateEditorStyle.cpp
+++ b/Engine/Source/Editor/EditorStyle/Private/SlateEditorStyle.cpp
@@ -1711,6 +1711,10 @@ void FSlateEditorStyle::FStyle::SetupGeneralStyles()
 		Set( "Sequencer.Timeline.PlayRange_Bottom_L", new BOX_BRUSH( "Sequencer/PlayRange_Bottom_L", FMargin(1.f, 0.5f, 0.f, 0.5f) ) );
 		Set( "Sequencer.Timeline.PlayRange_Bottom_R", new BOX_BRUSH( "Sequencer/PlayRange_Bottom_R", FMargin(0.f, 0.5f, 1.f, 0.5f) ) );
 
+		Set( "Sequencer.Timeline.SubSequenceRangeHashL", new BORDER_BRUSH( "Sequencer/SubSequenceRangeHashL", FMargin(1.f, 0.f, 0.f, 0.f) ) );
+		Set( "Sequencer.Timeline.SubSequenceRangeHashR", new BORDER_BRUSH( "Sequencer/SubSequenceRangeHashR", FMargin(1.f, 0.f, 0.f, 0.f) ) );
+		
+
 		
 		Set( "Sequencer.Transport.JumpToPreviousKey", FButtonStyle()
 			.SetNormal ( IMAGE_BRUSH( "/Sequencer/Transport_Bar/Previous_Frame_OFF", Icon24x24 ) )
@@ -1872,6 +1876,8 @@ void FSlateEditorStyle::FStyle::SetupGeneralStyles()
 		Set( "SequenceRecorder.Common.AddRecording", new IMAGE_BRUSH( "SequenceRecorder/icon_AddRecording_40x", Icon40x40 ) );
 		Set( "SequenceRecorder.Common.RemoveRecording.Small", new IMAGE_BRUSH( "SequenceRecorder/icon_RemoveRecording_40x", Icon20x20 ) );
 		Set( "SequenceRecorder.Common.RemoveRecording", new IMAGE_BRUSH( "SequenceRecorder/icon_RemoveRecording_40x", Icon40x40 ) );
+		Set( "SequenceRecorder.Common.RemoveAllRecordings.Small", new IMAGE_BRUSH( "SequenceRecorder/icon_RemoveRecording_40x", Icon20x20 ) );
+		Set( "SequenceRecorder.Common.RemoveAllRecordings", new IMAGE_BRUSH( "SequenceRecorder/icon_RemoveRecording_40x", Icon40x40 ) );
 	}
 
 	// Foliage Edit Mode
diff --git a/Engine/Source/Editor/EditorWidgets/Private/STransportControl.cpp b/Engine/Source/Editor/EditorWidgets/Private/STransportControl.cpp
--- a/Engine/Source/Editor/EditorWidgets/Private/STransportControl.cpp
+++ b/Engine/Source/Editor/EditorWidgets/Private/STransportControl.cpp
@@ -6,7 +6,7 @@
 #define LOCTEXT_NAMESPACE "STransportControl"
 
 BEGIN_SLATE_FUNCTION_BUILD_OPTIMIZATION
-TSharedPtr<SWidget> STransportControl::MakeTransportControlWidget(ETransportControlWidgetType WidgetType, const FOnMakeTransportWidget& MakeCustomWidgetDelegate)
+TSharedPtr<SWidget> STransportControl::MakeTransportControlWidget(ETransportControlWidgetType WidgetType, bool bAreButtonsFocusable, const FOnMakeTransportWidget& MakeCustomWidgetDelegate)
 {
 	switch(WidgetType)
 	{
@@ -16,14 +16,16 @@ TSharedPtr<SWidget> STransportControl::MakeTransportControlWidget(ETransportCont
 			. OnClicked(TransportControlArgs.OnBackwardEnd)
 			. Visibility(TransportControlArgs.OnBackwardEnd.IsBound() ? EVisibility::Visible : EVisibility::Collapsed)
 			. ToolTipText( LOCTEXT("ToFront", "To Front") )
-			. ContentPadding(2.0f);
+			. ContentPadding(2.0f)
+			. IsFocusable(bAreButtonsFocusable);
 	case ETransportControlWidgetType::BackwardStep:
 		return SNew(SButton)
 			. ButtonStyle(&FEditorStyle::Get().GetWidgetStyle<FButtonStyle>("Animation.Backward_Step"))
 			. OnClicked(TransportControlArgs.OnBackwardStep)
 			. Visibility(TransportControlArgs.OnBackwardStep.IsBound() ? EVisibility::Visible : EVisibility::Collapsed)
 			. ToolTipText( LOCTEXT("ToPrevious", "To Previous") )
-			. ContentPadding(2.0f);
+			. ContentPadding(2.0f)
+			. IsFocusable(bAreButtonsFocusable);
 	case ETransportControlWidgetType::BackwardPlay:
 		return SAssignNew(BackwardPlayButton, SButton)
 			. OnClicked(TransportControlArgs.OnBackwardPlay)
@@ -31,6 +33,7 @@ TSharedPtr<SWidget> STransportControl::MakeTransportControlWidget(ETransportCont
 			. ToolTipText( LOCTEXT("Reverse", "Reverse") )
 			. ButtonStyle( FEditorStyle::Get(), "NoBorder" )
 			. ContentPadding(2.0f)
+			. IsFocusable(bAreButtonsFocusable)
 			[
 				SNew(SImage)
 				. Image( this, &STransportControl::GetBackwardStatusIcon )
@@ -42,6 +45,7 @@ TSharedPtr<SWidget> STransportControl::MakeTransportControlWidget(ETransportCont
 			. Visibility(TransportControlArgs.OnRecord.IsBound() ? EVisibility::Visible : EVisibility::Collapsed)
 			. ToolTipText( this, &STransportControl::GetRecordStatusTooltip )
 			. ContentPadding(2.0f)
+			. IsFocusable(bAreButtonsFocusable)
 			[
 				SNew(SImage)
 				.Image(this, &STransportControl::GetRecordStatusIcon )
@@ -53,6 +57,7 @@ TSharedPtr<SWidget> STransportControl::MakeTransportControlWidget(ETransportCont
 			. ToolTipText( this, &STransportControl::GetForwardStatusTooltip )
 			. ButtonStyle( FEditorStyle::Get(), "NoBorder" )
 			. ContentPadding(2.0f)
+			. IsFocusable(bAreButtonsFocusable)
 			[
 				SNew(SImage)
 				. Image( this, &STransportControl::GetForwardStatusIcon )
@@ -63,14 +68,16 @@ TSharedPtr<SWidget> STransportControl::MakeTransportControlWidget(ETransportCont
 			. OnClicked(TransportControlArgs.OnForwardStep)
 			. Visibility(TransportControlArgs.OnForwardStep.IsBound() ? EVisibility::Visible : EVisibility::Collapsed)
 			. ToolTipText( LOCTEXT("ToNext", "To Next") )
-			. ContentPadding(2.0f);
+			. ContentPadding(2.0f)
+			. IsFocusable(bAreButtonsFocusable);
 	case ETransportControlWidgetType::ForwardEnd:
 		return SNew(SButton)
 			. ButtonStyle(&FEditorStyle::Get().GetWidgetStyle<FButtonStyle>("Animation.Forward_End"))
 			. OnClicked(TransportControlArgs.OnForwardEnd)
 			. Visibility(TransportControlArgs.OnForwardEnd.IsBound() ? EVisibility::Visible : EVisibility::Collapsed)
 			. ToolTipText( LOCTEXT("ToEnd", "To End") )
-			. ContentPadding(2.0f);
+			. ContentPadding(2.0f)
+			. IsFocusable(bAreButtonsFocusable);
 	case ETransportControlWidgetType::Loop:
 		return SNew(SButton)
 			. OnClicked(TransportControlArgs.OnToggleLooping)
@@ -78,6 +85,7 @@ TSharedPtr<SWidget> STransportControl::MakeTransportControlWidget(ETransportCont
 			. ToolTipText( LOCTEXT("Loop", "Loop") )
 			. ButtonStyle( FEditorStyle::Get(), "NoBorder" )
 			. ContentPadding(2.0f)
+			. IsFocusable(bAreButtonsFocusable)
 			[
 				SNew(SImage)
 				. Image( this, &STransportControl::GetLoopStatusIcon )
@@ -102,7 +110,7 @@ void STransportControl::Construct( const STransportControl::FArguments& InArgs )
 	{
 		for(FTransportControlWidget WidgetDesc : TransportControlArgs.WidgetsToCreate)
 		{
-			TSharedPtr<SWidget> Widget = MakeTransportControlWidget(WidgetDesc.WidgetType, WidgetDesc.MakeCustomWidgetDelegate);
+			TSharedPtr<SWidget> Widget = MakeTransportControlWidget(WidgetDesc.WidgetType, InArgs._TransportArgs.bAreButtonsFocusable, WidgetDesc.MakeCustomWidgetDelegate);
 			if(Widget.IsValid())
 			{
 				HorizontalBox->AddSlot()
@@ -119,7 +127,7 @@ void STransportControl::Construct( const STransportControl::FArguments& InArgs )
 	{
 		for(ETransportControlWidgetType WidgetType : TEnumRange<ETransportControlWidgetType>())
 		{
-			TSharedPtr<SWidget> Widget = MakeTransportControlWidget(WidgetType);
+			TSharedPtr<SWidget> Widget = MakeTransportControlWidget(WidgetType, InArgs._TransportArgs.bAreButtonsFocusable);
 			if(Widget.IsValid())
 			{
 				HorizontalBox->AddSlot()
diff --git a/Engine/Source/Editor/EditorWidgets/Private/STransportControl.h b/Engine/Source/Editor/EditorWidgets/Private/STransportControl.h
--- a/Engine/Source/Editor/EditorWidgets/Private/STransportControl.h
+++ b/Engine/Source/Editor/EditorWidgets/Private/STransportControl.h
@@ -42,7 +42,7 @@ private:
 	FReply OnToggleLooping();
 
 	/** Make default transport control widgets */
-	TSharedPtr<SWidget> MakeTransportControlWidget(ETransportControlWidgetType WidgetType, const FOnMakeTransportWidget& MakeCustomWidgetDelegate = FOnMakeTransportWidget());
+	TSharedPtr<SWidget> MakeTransportControlWidget(ETransportControlWidgetType WidgetType, bool bAreButtonsFocusable, const FOnMakeTransportWidget& MakeCustomWidgetDelegate = FOnMakeTransportWidget());
 
 private:
 	/** The handle to the active timer */
diff --git a/Engine/Source/Editor/EditorWidgets/Public/ITransportControl.h b/Engine/Source/Editor/EditorWidgets/Public/ITransportControl.h
--- a/Engine/Source/Editor/EditorWidgets/Public/ITransportControl.h
+++ b/Engine/Source/Editor/EditorWidgets/Public/ITransportControl.h
@@ -67,6 +67,7 @@ struct FTransportControlArgs
 		, OnGetPlaybackMode()
 		, OnTickPlayback()
 		, OnGetRecording()
+		, bAreButtonsFocusable(true)
 	{}
 
 	FOnClicked OnForwardPlay;
@@ -81,6 +82,7 @@ struct FTransportControlArgs
 	FOnGetPlaybackMode OnGetPlaybackMode;
 	FOnTickPlayback OnTickPlayback;
 	FOnGetRecording OnGetRecording;
+	bool bAreButtonsFocusable;
 
 	/** 
 	 * Array of custom widgets to create - if this array is used the default widget ordering will be
diff --git a/Engine/Source/Editor/Kismet/Private/BlueprintDetailsCustomization.cpp b/Engine/Source/Editor/Kismet/Private/BlueprintDetailsCustomization.cpp
--- a/Engine/Source/Editor/Kismet/Private/BlueprintDetailsCustomization.cpp
+++ b/Engine/Source/Editor/Kismet/Private/BlueprintDetailsCustomization.cpp
@@ -1690,8 +1690,9 @@ EVisibility FBlueprintVarActionDetails::ExposeToCinematicsVisibility() const
 		const bool bIsVectorStruct = VariableProperty->IsA(UStructProperty::StaticClass()) && Cast<UStructProperty>(VariableProperty)->Struct->GetFName() == NAME_Vector;
 		const bool bIsColorStruct = VariableProperty->IsA(UStructProperty::StaticClass()) && Cast<UStructProperty>(VariableProperty)->Struct->GetFName() == NAME_Color;
 		const bool bIsLinearColorStruct = VariableProperty->IsA(UStructProperty::StaticClass()) && Cast<UStructProperty>(VariableProperty)->Struct->GetFName() == NAME_LinearColor;
+		const bool bIsActorProperty = VariableProperty->IsA(UObjectProperty::StaticClass()) && Cast<UObjectProperty>(VariableProperty)->PropertyClass->IsChildOf(AActor::StaticClass());
 
-		if (bIsInteger || bIsByte || bIsFloat || bIsBool || bIsStr || bIsVectorStruct || bIsColorStruct || bIsLinearColorStruct)
+		if (bIsInteger || bIsByte || bIsFloat || bIsBool || bIsStr || bIsVectorStruct || bIsColorStruct || bIsLinearColorStruct || bIsActorProperty)
 		{
 			return EVisibility::Visible;
 		}
diff --git a/Engine/Source/Editor/MovieSceneTools/Private/CurveKeyAreas/FloatCurveKeyArea.cpp b/Engine/Source/Editor/MovieSceneTools/Private/CurveKeyAreas/FloatCurveKeyArea.cpp
--- a/Engine/Source/Editor/MovieSceneTools/Private/CurveKeyAreas/FloatCurveKeyArea.cpp
+++ b/Engine/Source/Editor/MovieSceneTools/Private/CurveKeyAreas/FloatCurveKeyArea.cpp
@@ -226,7 +226,7 @@ void FFloatCurveKeyArea::SetKeyTangentMode(FKeyHandle KeyHandle, ERichCurveTange
 }
 
 
-void FFloatCurveKeyArea::SetKeyTime(FKeyHandle KeyHandle, float NewKeyTime) const
+void FFloatCurveKeyArea::SetKeyTime(FKeyHandle KeyHandle, float NewKeyTime)
 {
 	Curve->SetKeyTime(KeyHandle, NewKeyTime);
 }
diff --git a/Engine/Source/Editor/MovieSceneTools/Private/CurveKeyAreas/IntegerKeyArea.cpp b/Engine/Source/Editor/MovieSceneTools/Private/CurveKeyAreas/IntegerKeyArea.cpp
new file mode 100644
--- /dev/null
+++ b/Engine/Source/Editor/MovieSceneTools/Private/CurveKeyAreas/IntegerKeyArea.cpp
@@ -0,0 +1,24 @@
+// Copyright 1998-2016 Epic Games, Inc. All Rights Reserved.
+
+#include "MovieSceneToolsPrivatePCH.h"
+#include "IntegerKeyArea.h"
+#include "SIntegralCurveKeyEditor.h"
+
+
+/* IKeyArea interface
+ *****************************************************************************/
+
+bool FIntegerKeyArea::CanCreateKeyEditor()
+{
+	return true;
+}
+
+
+TSharedRef<SWidget> FIntegerKeyArea::CreateKeyEditor(ISequencer* Sequencer)
+{
+	return SNew(SIntegralCurveKeyEditor<int32>)
+		.Sequencer(Sequencer)
+		.OwningSection(OwningSection)
+		.Curve(&Curve)
+		.ExternalValue(ExternalValue);
+};
diff --git a/Engine/Source/Editor/MovieSceneTools/Private/CurveKeyAreas/IntegralKeyArea.cpp b/Engine/Source/Editor/MovieSceneTools/Private/CurveKeyAreas/IntegralKeyArea.cpp
--- a/Engine/Source/Editor/MovieSceneTools/Private/CurveKeyAreas/IntegralKeyArea.cpp
+++ b/Engine/Source/Editor/MovieSceneTools/Private/CurveKeyAreas/IntegralKeyArea.cpp
@@ -136,7 +136,7 @@ void FIntegralCurveKeyAreaBase::SetKeyTangentMode(FKeyHandle KeyHandle, ERichCur
 }
 
 
-void FIntegralCurveKeyAreaBase::SetKeyTime(FKeyHandle KeyHandle, float NewKeyTime) const
+void FIntegralCurveKeyAreaBase::SetKeyTime(FKeyHandle KeyHandle, float NewKeyTime)
 {
 	Curve.SetKeyTime(KeyHandle, NewKeyTime);
 }
diff --git a/Engine/Source/Editor/MovieSceneTools/Private/CurveKeyAreas/NameCurveKeyArea.cpp b/Engine/Source/Editor/MovieSceneTools/Private/CurveKeyAreas/NameCurveKeyArea.cpp
--- a/Engine/Source/Editor/MovieSceneTools/Private/CurveKeyAreas/NameCurveKeyArea.cpp
+++ b/Engine/Source/Editor/MovieSceneTools/Private/CurveKeyAreas/NameCurveKeyArea.cpp
@@ -145,7 +145,7 @@ void FNameCurveKeyArea::SetKeyTangentMode(FKeyHandle KeyHandle, ERichCurveTangen
 }
 
 
-void FNameCurveKeyArea::SetKeyTime(FKeyHandle KeyHandle, float NewKeyTime) const
+void FNameCurveKeyArea::SetKeyTime(FKeyHandle KeyHandle, float NewKeyTime)
 {
 	Curve.SetKeyTime(KeyHandle, NewKeyTime);
 }
diff --git a/Engine/Source/Editor/MovieSceneTools/Private/CurveKeyAreas/StringCurveKeyArea.cpp b/Engine/Source/Editor/MovieSceneTools/Private/CurveKeyAreas/StringCurveKeyArea.cpp
--- a/Engine/Source/Editor/MovieSceneTools/Private/CurveKeyAreas/StringCurveKeyArea.cpp
+++ b/Engine/Source/Editor/MovieSceneTools/Private/CurveKeyAreas/StringCurveKeyArea.cpp
@@ -180,7 +180,7 @@ void FStringCurveKeyArea::SetKeyTangentMode(FKeyHandle KeyHandle, ERichCurveTang
 }
 
 
-void FStringCurveKeyArea::SetKeyTime(FKeyHandle KeyHandle, float NewKeyTime) const
+void FStringCurveKeyArea::SetKeyTime(FKeyHandle KeyHandle, float NewKeyTime)
 {
 	Curve->SetKeyTime(KeyHandle, NewKeyTime);
 }
diff --git a/Engine/Source/Editor/MovieSceneTools/Private/CurveKeyEditors/SBoolCurveKeyEditor.cpp b/Engine/Source/Editor/MovieSceneTools/Private/CurveKeyEditors/SBoolCurveKeyEditor.cpp
--- a/Engine/Source/Editor/MovieSceneTools/Private/CurveKeyEditors/SBoolCurveKeyEditor.cpp
+++ b/Engine/Source/Editor/MovieSceneTools/Private/CurveKeyEditors/SBoolCurveKeyEditor.cpp
@@ -29,7 +29,7 @@ ECheckBoxState SBoolCurveKeyEditor::IsChecked() const
 	}
 	else
 	{
-		float CurrentTime = Sequencer->GetCurrentLocalTime(*Sequencer->GetFocusedMovieSceneSequence());
+		float CurrentTime = Sequencer->GetLocalTime();
 		bool DefaultValue = false;
 		bCurrentValue = Curve->Evaluate(CurrentTime, DefaultValue) != 0;
 	}
@@ -43,7 +43,7 @@ void SBoolCurveKeyEditor::OnCheckStateChanged(ECheckBoxState NewCheckboxState)
 	OwningSection->SetFlags(RF_Transactional);
 	if (OwningSection->TryModify())
 	{
-		float CurrentTime = Sequencer->GetCurrentLocalTime(*Sequencer->GetFocusedMovieSceneSequence());
+		float CurrentTime = Sequencer->GetLocalTime();
 		bool bAutoSetTrackDefaults = Sequencer->GetAutoSetTrackDefaults();
 		int32 NewValue = NewCheckboxState == ECheckBoxState::Checked ? 1 : 0;
 
diff --git a/Engine/Source/Editor/MovieSceneTools/Private/CurveKeyEditors/SEnumCurveKeyEditor.cpp b/Engine/Source/Editor/MovieSceneTools/Private/CurveKeyEditors/SEnumCurveKeyEditor.cpp
--- a/Engine/Source/Editor/MovieSceneTools/Private/CurveKeyEditors/SEnumCurveKeyEditor.cpp
+++ b/Engine/Source/Editor/MovieSceneTools/Private/CurveKeyEditors/SEnumCurveKeyEditor.cpp
@@ -33,7 +33,7 @@ int32 SEnumCurveKeyEditor::OnGetCurrentValue() const
 		return ExternalValue.Get().GetValue();
 	}
 
-	float CurrentTime = Sequencer->GetCurrentLocalTime(*Sequencer->GetFocusedMovieSceneSequence());
+	float CurrentTime = Sequencer->GetLocalTime();
 	int32 DefaultValue = 0;
 	return Curve->Evaluate(CurrentTime, DefaultValue);
 }
@@ -46,7 +46,7 @@ void SEnumCurveKeyEditor::OnComboSelectionChanged(int32 InSelectedItem, ESelectI
 
 	if (OwningSection->TryModify())
 	{
-		float CurrentTime = Sequencer->GetCurrentLocalTime(*Sequencer->GetFocusedMovieSceneSequence());
+		float CurrentTime = Sequencer->GetLocalTime();
 		bool bAutoSetTrackDefaults = Sequencer->GetAutoSetTrackDefaults();
 
 		FKeyHandle CurrentKeyHandle = Curve->FindKey(CurrentTime);
diff --git a/Engine/Source/Editor/MovieSceneTools/Private/CurveKeyEditors/SStringCurveKeyEditor.cpp b/Engine/Source/Editor/MovieSceneTools/Private/CurveKeyEditors/SStringCurveKeyEditor.cpp
--- a/Engine/Source/Editor/MovieSceneTools/Private/CurveKeyEditors/SStringCurveKeyEditor.cpp
+++ b/Engine/Source/Editor/MovieSceneTools/Private/CurveKeyEditors/SStringCurveKeyEditor.cpp
@@ -13,7 +13,7 @@ void SStringCurveKeyEditor::Construct(const FArguments& InArgs)
 	Curve = InArgs._Curve;
 	ExternalValue = InArgs._ExternalValue;
 
-	float CurrentTime = Sequencer->GetCurrentLocalTime(*Sequencer->GetFocusedMovieSceneSequence());
+	float CurrentTime = Sequencer->GetLocalTime();
 	ChildSlot
 	[
 		SNew(SEditableText)
@@ -29,7 +29,7 @@ FText SStringCurveKeyEditor::GetText() const
 	{
 		return FText::FromString(ExternalValue.Get().GetValue());
 	}
-	float CurrentTime = Sequencer->GetCurrentLocalTime(*Sequencer->GetFocusedMovieSceneSequence());
+	float CurrentTime = Sequencer->GetLocalTime();
 	FString DefaultValue;
 	FString CurrentValue = Curve->Eval(CurrentTime, DefaultValue);
 	
@@ -42,7 +42,7 @@ void SStringCurveKeyEditor::OnTextCommitted(const FText& InText, ETextCommit::Ty
 	OwningSection->SetFlags(RF_Transactional);
 	if (OwningSection->TryModify())
 	{
-		float CurrentTime = Sequencer->GetCurrentLocalTime(*Sequencer->GetFocusedMovieSceneSequence());
+		float CurrentTime = Sequencer->GetLocalTime();
 		bool bAutoSetTrackDefaults = Sequencer->GetAutoSetTrackDefaults();
 
 		FKeyHandle CurrentKeyHandle = Curve->FindKey(CurrentTime);
diff --git a/Engine/Source/Editor/MovieSceneTools/Private/MatineeImportTools.cpp b/Engine/Source/Editor/MovieSceneTools/Private/MatineeImportTools.cpp
--- a/Engine/Source/Editor/MovieSceneTools/Private/MatineeImportTools.cpp
+++ b/Engine/Source/Editor/MovieSceneTools/Private/MatineeImportTools.cpp
@@ -40,7 +40,6 @@
 #include "MovieSceneFadeSection.h"
 #include "MovieSceneCameraCutSection.h"
 #include "MovieSceneEventSection.h"
-#include "MovieSceneVisibilitySection.h"
 
 #include "MovieSceneSequence.h"
 
@@ -499,11 +498,11 @@ bool FMatineeImportTools::CopyInterpAnimControlTrack( UInterpTrackAnimControl* M
 		UMovieSceneSkeletalAnimationSection* NewSection = Cast<UMovieSceneSkeletalAnimationSection>( SkeletalAnimationTrack->CreateNewSection() );
 		NewSection->SetStartTime( AnimSeq.StartTime );
 		NewSection->SetEndTime( EndTime );
-		NewSection->SetStartOffset( AnimSeq.AnimStartOffset );
-		NewSection->SetEndOffset( AnimSeq.AnimEndOffset );
-		NewSection->SetPlayRate( AnimSeq.AnimPlayRate );
-		NewSection->SetAnimSequence( AnimSeq.AnimSeq );
-		NewSection->SetSlotName( MatineeAnimControlTrack->SlotName );
+		NewSection->Params.StartOffset = AnimSeq.AnimStartOffset;
+		NewSection->Params.EndOffset = AnimSeq.AnimEndOffset;
+		NewSection->Params.PlayRate = AnimSeq.AnimPlayRate;
+		NewSection->Params.Animation = AnimSeq.AnimSeq;
+		NewSection->Params.SlotName = MatineeAnimControlTrack->SlotName;
 
 		SkeletalAnimationTrack->AddSection( *NewSection );
 		bSectionCreated = true;
@@ -580,7 +579,7 @@ bool FMatineeImportTools::CopyInterpFadeTrack( UInterpTrackFade* MatineeFadeTrac
 	return bSectionCreated;
 }
 
-bool FMatineeImportTools::CopyInterpDirectorTrack( UInterpTrackDirector* DirectorTrack, UMovieSceneCameraCutTrack* CameraCutTrack, AMatineeActor* MatineeActor, UMovieSceneSequence* MovieSceneSequence )
+bool FMatineeImportTools::CopyInterpDirectorTrack( UInterpTrackDirector* DirectorTrack, UMovieSceneCameraCutTrack* CameraCutTrack, AMatineeActor* MatineeActor, IMovieScenePlayer& Player )
 {
 	const FScopedTransaction Transaction( NSLOCTEXT( "Sequencer", "PasteMatineeDirectorTrack", "Paste Matinee Director Track" ) );
 	bool bCutsAdded = false;
@@ -600,7 +599,7 @@ bool FMatineeImportTools::CopyInterpDirectorTrack( UInterpTrackDirector* Directo
 			{
 				AActor* CameraActor = ViewGroupInst->GetGroupActor();
 		
-				FGuid CameraHandle = MovieSceneSequence->FindPossessableObjectId(*CameraActor);
+				FGuid CameraHandle = Player.FindObjectId(*CameraActor, MovieSceneSequenceID::Root);
 				if (CameraHandle.IsValid())
 				{
 					CameraCutTrack->AddNewCameraCut(CameraHandle, TrackCut.Time);
@@ -634,11 +633,11 @@ bool FMatineeImportTools::CopyInterpEventTrack( UInterpTrackEvent* MatineeEventT
 		{
 			float SectionMin = Section->GetStartTime();
 			float SectionMax = Section->GetEndTime();
-			
-			FNameCurve& EventCurve = Section->GetEventCurve();
+
+			TCurveInterface<FEventPayload, float> CurveInterface = Section->GetCurveInterface();
 			for (FEventTrackKey EventTrackKey : MatineeEventTrack->EventTrack)
 			{
-				EventCurve.UpdateOrAddKey(EventTrackKey.Time, EventTrackKey.EventName);
+				CurveInterface.UpdateOrAddKey(EventTrackKey.Time, FEventPayload(EventTrackKey.EventName), KINDA_SMALL_NUMBER);
 				SectionMin = FMath::Min( SectionMin, EventTrackKey.Time );
 				SectionMax = FMath::Max( SectionMax, EventTrackKey.Time );
 			}
@@ -661,10 +660,10 @@ bool FMatineeImportTools::CopyInterpVisibilityTrack( UInterpTrackVisibility* Mat
 	if (MatineeVisibilityTrack->VisibilityTrack.Num())
 	{
 		float KeyTime = MatineeVisibilityTrack->VisibilityTrack[0].Time;
-		UMovieSceneVisibilitySection* Section = Cast<UMovieSceneVisibilitySection>( MovieSceneHelpers::FindSectionAtTime( VisibilityTrack->GetAllSections(), KeyTime ) );
+		UMovieSceneBoolSection* Section = Cast<UMovieSceneBoolSection>( MovieSceneHelpers::FindSectionAtTime( VisibilityTrack->GetAllSections(), KeyTime ) );
 		if ( Section == nullptr )
 		{
-			Section = Cast<UMovieSceneVisibilitySection>( VisibilityTrack->CreateNewSection() );
+			Section = Cast<UMovieSceneBoolSection>( VisibilityTrack->CreateNewSection() );
 			VisibilityTrack->AddSection( *Section );
 			bSectionCreated = true;
 		}
diff --git a/Engine/Source/Editor/MovieSceneTools/Private/MovieSceneToolHelpers.cpp b/Engine/Source/Editor/MovieSceneTools/Private/MovieSceneToolHelpers.cpp
--- a/Engine/Source/Editor/MovieSceneTools/Private/MovieSceneToolHelpers.cpp
+++ b/Engine/Source/Editor/MovieSceneTools/Private/MovieSceneToolHelpers.cpp
@@ -589,42 +589,45 @@ bool MovieSceneToolHelpers::ShowExportEDLDialog(const UMovieScene* InMovieScene,
 	return false;
 }
 
-bool ImportFBXProperty(FString NodeName, FString AnimatedPropertyName, FGuid ObjectBinding, UnFbx::FFbxCurvesAPI& CurveAPI, UMovieScene* InMovieScene, FMovieSceneSequenceInstance& InSequence, ISequencer& InSequencer)
+bool ImportFBXProperty(FString NodeName, FString AnimatedPropertyName, FGuid ObjectBinding, UnFbx::FFbxCurvesAPI& CurveAPI, UMovieScene* InMovieScene, ISequencer& InSequencer)
 {
 	const UMovieSceneToolsProjectSettings* ProjectSettings = GetDefault<UMovieSceneToolsProjectSettings>();
 
+	TArrayView<TWeakObjectPtr<>> BoundObjects = InSequencer.FindBoundObjects(ObjectBinding, InSequencer.GetFocusedTemplateID());
+
 	for (auto FbxSetting : ProjectSettings->FbxSettings)
 	{
 		if (FCString::Strcmp(*FbxSetting.FbxPropertyName.ToUpper(), *AnimatedPropertyName.ToUpper()) != 0)
 		{
 			continue;
 		}
-					
-		UObject* FoundObject = InSequence.FindObject(ObjectBinding, InSequencer);
-		if (!FoundObject)
-		{
-			continue;
-		}
-		
-		UObject* PropertyOwner = FoundObject;
-		if (!FbxSetting.PropertyPath.ComponentName.IsEmpty())
-		{
-			PropertyOwner = FindObjectFast<UObject>(FoundObject, *FbxSetting.PropertyPath.ComponentName);
-		}
 
-		if (!PropertyOwner)
-		{
-			continue;
-		}
-	
-		FGuid PropertyOwnerGuid = InSequence.FindObjectId(*PropertyOwner);
-		if (!PropertyOwnerGuid.IsValid())
+		for (TWeakObjectPtr<>& WeakObject : BoundObjects)
 		{
-			PropertyOwnerGuid = InSequencer.GetHandleToObject(PropertyOwner);
-		}
+			UObject* FoundObject = WeakObject.Get();
+
+			if (!FoundObject)
+			{
+				continue;
+			}
+			
+			UObject* PropertyOwner = FoundObject;
+			if (!FbxSetting.PropertyPath.ComponentName.IsEmpty())
+			{
+				PropertyOwner = FindObjectFast<UObject>(FoundObject, *FbxSetting.PropertyPath.ComponentName);
+			}
+
+			if (!PropertyOwner)
+			{
+				continue;
+			}
+		
+			FGuid PropertyOwnerGuid = InSequencer.GetHandleToObject(PropertyOwner);
+			if (!PropertyOwnerGuid.IsValid())
+			{
+				continue;
+			}
 
-		if (PropertyOwnerGuid.IsValid())
-		{
 			UMovieSceneFloatTrack* FloatTrack = InMovieScene->FindTrack<UMovieSceneFloatTrack>(PropertyOwnerGuid, *FbxSetting.PropertyPath.PropertyName);
 			if (!FloatTrack)
 			{
@@ -635,7 +638,7 @@ bool ImportFBXProperty(FString NodeName, FString AnimatedPropertyName, FGuid Obj
 			}
 
 			if (FloatTrack)
-			{			
+			{
 				bool bSectionAdded = false;
 				UMovieSceneFloatSection* FloatSection = Cast<UMovieSceneFloatSection>(FloatTrack->FindOrAddSection(0.f, bSectionAdded));
 				if (!FloatSection)
@@ -812,7 +815,7 @@ bool ImportFBXTransform(FString NodeName, FGuid ObjectBinding, UnFbx::FFbxCurves
 	return true;
 }
 
-bool ImportFBXNode(FString NodeName, UnFbx::FFbxCurvesAPI& CurveAPI, UMovieScene* InMovieScene, FMovieSceneSequenceInstance& InSequence, ISequencer& InSequencer, const TMap<FGuid, FString>& InObjectBindingMap)
+bool ImportFBXNode(FString NodeName, UnFbx::FFbxCurvesAPI& CurveAPI, UMovieScene* InMovieScene, ISequencer& InSequencer, const TMap<FGuid, FString>& InObjectBindingMap)
 {
 	// Find the matching object binding to apply this animation to. Defaults to the first.
 	FGuid ObjectBinding;
@@ -837,7 +840,7 @@ bool ImportFBXNode(FString NodeName, UnFbx::FFbxCurvesAPI& CurveAPI, UMovieScene
 		
 	for (auto AnimatedPropertyName : AnimatedPropertyNames)
 	{
-		ImportFBXProperty(NodeName, AnimatedPropertyName, ObjectBinding, CurveAPI, InMovieScene, InSequence, InSequencer);
+		ImportFBXProperty(NodeName, AnimatedPropertyName, ObjectBinding, CurveAPI, InMovieScene, InSequencer);
 	}
 	
 	ImportFBXTransform(NodeName, ObjectBinding, CurveAPI, InMovieScene);
@@ -845,7 +848,7 @@ bool ImportFBXNode(FString NodeName, UnFbx::FFbxCurvesAPI& CurveAPI, UMovieScene
 	return true;
 }
 
-bool MovieSceneToolHelpers::ImportFBX(UMovieScene* InMovieScene, FMovieSceneSequenceInstance& InSequence, ISequencer& InSequencer, const TMap<FGuid, FString>& InObjectBindingMap)
+bool MovieSceneToolHelpers::ImportFBX(UMovieScene* InMovieScene, ISequencer& InSequencer, const TMap<FGuid, FString>& InObjectBindingMap)
 {
 	TArray<FString> OpenFilenames;
 	IDesktopPlatform* DesktopPlatform = FDesktopPlatformModule::Get();
@@ -897,7 +900,7 @@ bool MovieSceneToolHelpers::ImportFBX(UMovieScene* InMovieScene, FMovieSceneSequ
 
 	for (FString NodeName : AnimatedNodeNames)
 	{
-		ImportFBXNode(NodeName, CurveAPI, InMovieScene, InSequence, InSequencer, InObjectBindingMap);
+		ImportFBXNode(NodeName, CurveAPI, InMovieScene, InSequencer, InObjectBindingMap);
 	}
 
 	FbxImporter->ReleaseScene();
diff --git a/Engine/Source/Editor/MovieSceneTools/Private/MovieSceneToolsModule.cpp b/Engine/Source/Editor/MovieSceneTools/Private/MovieSceneToolsModule.cpp
--- a/Engine/Source/Editor/MovieSceneTools/Private/MovieSceneToolsModule.cpp
+++ b/Engine/Source/Editor/MovieSceneTools/Private/MovieSceneToolsModule.cpp
@@ -15,6 +15,7 @@
 #include "BytePropertyTrackEditor.h"
 #include "ColorPropertyTrackEditor.h"
 #include "FloatPropertyTrackEditor.h"
+#include "IntegerPropertyTrackEditor.h"
 #include "VectorPropertyTrackEditor.h"
 #include "VisibilityPropertyTrackEditor.h"
 #include "ActorReferencePropertyTrackEditor.h"
@@ -80,6 +81,7 @@ public:
 		BytePropertyTrackCreateEditorHandle = SequencerModule.RegisterTrackEditor_Handle( FOnCreateTrackEditor::CreateStatic( &FBytePropertyTrackEditor::CreateTrackEditor ) );
 		ColorPropertyTrackCreateEditorHandle = SequencerModule.RegisterTrackEditor_Handle( FOnCreateTrackEditor::CreateStatic( &FColorPropertyTrackEditor::CreateTrackEditor ) );
 		FloatPropertyTrackCreateEditorHandle = SequencerModule.RegisterTrackEditor_Handle( FOnCreateTrackEditor::CreateStatic( &FFloatPropertyTrackEditor::CreateTrackEditor ) );
+		IntegerPropertyTrackCreateEditorHandle = SequencerModule.RegisterTrackEditor_Handle( FOnCreateTrackEditor::CreateStatic( &FIntegerPropertyTrackEditor::CreateTrackEditor ) );
 		VectorPropertyTrackCreateEditorHandle = SequencerModule.RegisterTrackEditor_Handle( FOnCreateTrackEditor::CreateStatic( &FVectorPropertyTrackEditor::CreateTrackEditor ) );
 		VisibilityPropertyTrackCreateEditorHandle = SequencerModule.RegisterTrackEditor_Handle( FOnCreateTrackEditor::CreateStatic( &FVisibilityPropertyTrackEditor::CreateTrackEditor ) );
 		ActorReferencePropertyTrackCreateEditorHandle = SequencerModule.RegisterTrackEditor_Handle(FOnCreateTrackEditor::CreateStatic( &FActorReferencePropertyTrackEditor::CreateTrackEditor ) );
@@ -131,6 +133,7 @@ public:
 		SequencerModule.UnRegisterTrackEditor_Handle( BytePropertyTrackCreateEditorHandle );
 		SequencerModule.UnRegisterTrackEditor_Handle( ColorPropertyTrackCreateEditorHandle );
 		SequencerModule.UnRegisterTrackEditor_Handle( FloatPropertyTrackCreateEditorHandle );
+		SequencerModule.UnRegisterTrackEditor_Handle( IntegerPropertyTrackCreateEditorHandle );
 		SequencerModule.UnRegisterTrackEditor_Handle( VectorPropertyTrackCreateEditorHandle );
 		SequencerModule.UnRegisterTrackEditor_Handle( VisibilityPropertyTrackCreateEditorHandle );
 		SequencerModule.UnRegisterTrackEditor_Handle( ActorReferencePropertyTrackCreateEditorHandle );
@@ -203,6 +206,7 @@ private:
 	FDelegateHandle BytePropertyTrackCreateEditorHandle;
 	FDelegateHandle ColorPropertyTrackCreateEditorHandle;
 	FDelegateHandle FloatPropertyTrackCreateEditorHandle;
+	FDelegateHandle IntegerPropertyTrackCreateEditorHandle;
 	FDelegateHandle VectorPropertyTrackCreateEditorHandle;
 	FDelegateHandle VisibilityPropertyTrackCreateEditorHandle;
 	FDelegateHandle ActorReferencePropertyTrackCreateEditorHandle;
diff --git a/Engine/Source/Editor/MovieSceneTools/Private/PropertySection.cpp b/Engine/Source/Editor/MovieSceneTools/Private/PropertySection.cpp
--- a/Engine/Source/Editor/MovieSceneTools/Private/PropertySection.cpp
+++ b/Engine/Source/Editor/MovieSceneTools/Private/PropertySection.cpp
@@ -41,34 +41,22 @@ int32 FPropertySection::OnPaintSection( FSequencerSectionPainter& Painter ) cons
 	return Painter.PaintSectionBackground();
 }
 
-UObject* FPropertySection::GetRuntimeObjectAndUpdatePropertyBindings() const
+UProperty* FPropertySection::GetProperty() const
 {
-	UObject* RuntimeObject = nullptr;
-	if (Sequencer != nullptr && PropertyBindings.IsValid())
+	if (!PropertyBindings.IsValid())
+	{
+		return nullptr;
+	}
+
+	for (const TWeakObjectPtr<UObject>& WeakObject : Sequencer->FindBoundObjects(ObjectBinding, Sequencer->GetFocusedTemplateID()))
 	{
-		TArray<TWeakObjectPtr<UObject>> RuntimeObjects;
-		Sequencer->GetRuntimeObjects(Sequencer->GetFocusedMovieSceneSequenceInstance(), ObjectBinding, RuntimeObjects);
-		if (RuntimeObjects.Num() == 1)
+		if (UObject* Object = WeakObject.Get())
 		{
-			TWeakObjectPtr<UObject> RuntimeObjectPtr = RuntimeObjects[0];
-			if (RuntimeObjectPtr.IsValid())
-			{
-				if (RuntimeObjectPtr != RuntimeObjectCache)
-				{
-					PropertyBindings->UpdateBindings(RuntimeObjects);
-					RuntimeObjectCache = RuntimeObjectPtr;
-				}
-				RuntimeObject = RuntimeObjectPtr.Get();
-			}
+			return PropertyBindings->GetProperty(*Object);
 		}
 	}
-	return RuntimeObject;
-}
 
-UProperty* FPropertySection::GetProperty() const
-{
-	UObject* RuntimeObject = GetRuntimeObjectAndUpdatePropertyBindings();
-	return PropertyBindings.IsValid() && RuntimeObject != nullptr ? PropertyBindings->GetProperty(RuntimeObject) : nullptr;
+	return nullptr;
 }
 
 bool FPropertySection::CanGetPropertyValue() const
diff --git a/Engine/Source/Editor/MovieSceneTools/Private/Sections/ActorReferencePropertySection.cpp b/Engine/Source/Editor/MovieSceneTools/Private/Sections/ActorReferencePropertySection.cpp
--- a/Engine/Source/Editor/MovieSceneTools/Private/Sections/ActorReferencePropertySection.cpp
+++ b/Engine/Source/Editor/MovieSceneTools/Private/Sections/ActorReferencePropertySection.cpp
@@ -128,7 +128,8 @@ TOptional<FGuid> FActorReferencePropertySection::GetActorGuid() const
 	TOptional<AActor*> CurrentActor = GetPropertyValue<AActor*>();
 	if (CurrentActor.IsSet() && CurrentActor.GetValue() != nullptr)
 	{
-		return TOptional<FGuid>(GetSequencer()->GetFocusedMovieSceneSequenceInstance()->FindObjectId(*CurrentActor.GetValue()));
+		ISequencer* SequencerPtr = GetSequencer();
+		return TOptional<FGuid>(SequencerPtr->FindObjectId(*CurrentActor.GetValue(), SequencerPtr->GetFocusedTemplateID()));
 	}
 	return TOptional<FGuid>(FGuid());
 }
\ No newline at end of file
diff --git a/Engine/Source/Editor/MovieSceneTools/Private/Sections/CameraCutSection.cpp b/Engine/Source/Editor/MovieSceneTools/Private/Sections/CameraCutSection.cpp
--- a/Engine/Source/Editor/MovieSceneTools/Private/Sections/CameraCutSection.cpp
+++ b/Engine/Source/Editor/MovieSceneTools/Private/Sections/CameraCutSection.cpp
@@ -123,21 +123,19 @@ const AActor* FCameraCutSection::GetCameraForFrame(float Time) const
 
 	if (CameraCutSection && Sequencer.IsValid())
 	{
-		TSharedRef<FMovieSceneSequenceInstance> SequenceInstance = Sequencer->GetFocusedMovieSceneSequenceInstance();
-
-		AActor* Actor = Cast<AActor>(SequenceInstance->FindObject(CameraCutSection->GetCameraGuid(), *Sequencer));
-		if (Actor)
-		{
-			return Actor;
-		}
-		else
+		for (TWeakObjectPtr<>& Object : Sequencer->FindBoundObjects(CameraCutSection->GetCameraGuid(), Sequencer->GetFocusedTemplateID()))
 		{
-			FMovieSceneSpawnable* Spawnable = SequenceInstance->GetSequence()->GetMovieScene()->FindSpawnable(CameraCutSection->GetCameraGuid());
-			if (Spawnable)
+			if (AActor* Actor = Cast<AActor>(Object.Get()))
 			{
-				return Cast<AActor>(Spawnable->GetObjectTemplate());
+				return Actor;
 			}
 		}
+
+		FMovieSceneSpawnable* Spawnable = Sequencer->GetFocusedMovieSceneSequence()->GetMovieScene()->FindSpawnable(CameraCutSection->GetCameraGuid());
+		if (Spawnable)
+		{
+			return Cast<AActor>(Spawnable->GetObjectTemplate());
+		}
 	}
 
 	return nullptr;
