Add LLM tags to session profiler.

#rnx

diff --git a/Engine/Source/Developer/Profiler/Private/ProfilerModule.cpp b/Engine/Source/Developer/Profiler/Private/ProfilerModule.cpp
--- a/Engine/Source/Developer/Profiler/Private/ProfilerModule.cpp
+++ b/Engine/Source/Developer/Profiler/Private/ProfilerModule.cpp
@@ -55,6 +55,7 @@ IMPLEMENT_MODULE( FProfilerModule, Profiler );
 TSharedRef<SWidget> FProfilerModule::CreateProfilerWindow( const TSharedRef<ISessionManager>& InSessionManager, const TSharedRef<SDockTab>& ConstructUnderMajorTab )
 {
 #if STATS
+	LLM_SCOPE_BYNAME(TEXT("SessionProfiler"));
 	FProfilerManager::Initialize( InSessionManager );
 	TSharedRef<SProfilerWindow> ProfilerWindow = SNew( SProfilerWindow );
 	FProfilerManager::Get()->AssignProfilerWindow( ProfilerWindow );
diff --git a/Engine/Source/Developer/ProfilerClient/Private/ProfilerClientManager.cpp b/Engine/Source/Developer/ProfilerClient/Private/ProfilerClientManager.cpp
--- a/Engine/Source/Developer/ProfilerClient/Private/ProfilerClientManager.cpp
+++ b/Engine/Source/Developer/ProfilerClient/Private/ProfilerClientManager.cpp
@@ -526,6 +526,8 @@ bool FProfilerClientManager::CheckHashAndWrite(const FProfilerServiceFileChunk&
 void FProfilerClientManager::HandleServiceFileChunk(const FProfilerServiceFileChunk& FileChunk, const TSharedRef<IMessageContext, ESPMode::ThreadSafe>& Context)
 {
 #if STATS
+	LLM_SCOPE_BYNAME(TEXT("SessionProfiler"));
+
 	const TCHAR* StrTmp = TEXT(".tmp");
 
 	// Read file chunk header.
@@ -773,6 +775,8 @@ void FProfilerClientManager::HandleProfilerServiceData2Message(const FProfilerSe
 	SCOPE_CYCLE_COUNTER(STAT_PC_HandleDataReceived);
 	if (ActiveSessionId.IsValid() && Connections.Find(Message.InstanceId) != nullptr)
 	{
+		LLM_SCOPE_BYNAME(TEXT("SessionProfiler"));
+
 		// Create a temporary profiler data and prepare all data.
 		FProfilerServiceData2* ToProcess = new FProfilerServiceData2(Message.InstanceId, Message.Frame, Message.HexData, Message.CompressedSize, Message.UncompressedSize);
 
diff --git a/Engine/Source/Developer/ProfilerService/Private/ProfilerServiceFileTransfer.cpp b/Engine/Source/Developer/ProfilerService/Private/ProfilerServiceFileTransfer.cpp
--- a/Engine/Source/Developer/ProfilerService/Private/ProfilerServiceFileTransfer.cpp
+++ b/Engine/Source/Developer/ProfilerService/Private/ProfilerServiceFileTransfer.cpp
@@ -54,6 +54,7 @@ bool FFileTransferRunnable::Init()
 
 uint32 FFileTransferRunnable::Run()
 {
+	LLM_SCOPE_BYNAME(TEXT("SessionProfiler"));
 	while (!ShouldStop())
 	{
 		if (WorkEvent->Wait(250))
diff --git a/Engine/Source/Developer/ProfilerService/Private/ProfilerServiceManager.cpp b/Engine/Source/Developer/ProfilerService/Private/ProfilerServiceManager.cpp
--- a/Engine/Source/Developer/ProfilerService/Private/ProfilerServiceManager.cpp
+++ b/Engine/Source/Developer/ProfilerService/Private/ProfilerServiceManager.cpp
@@ -129,6 +129,7 @@ TSharedPtr<IProfilerServiceManager> FProfilerServiceManager::CreateSharedService
 void FProfilerServiceManager::AddNewFrameHandleStatsPipe()
 {
 #if	STATS
+	LLM_SCOPE_BYNAME(TEXT("SessionProfiler"));
 	const FStatsThreadState& Stats = FStatsThreadState::GetLocalState();
 	NewFrameDelegateHandle = Stats.NewFrameDelegate.AddRaw( this, &FProfilerServiceManager::HandleNewFrame );
 	StatsPrimaryEnableAdd();
@@ -152,6 +153,8 @@ extern CORE_API UE::Tasks::FPipe GStatsPipe;
 void FProfilerServiceManager::SetPreviewState( const FMessageAddress& ClientAddress, const bool bRequestedPreviewState )
 {
 #if STATS
+	LLM_SCOPE_BYNAME(TEXT("SessionProfiler"));
+
 	FClientData* Client = ClientData.Find( ClientAddress );
 	if (MessageEndpoint.IsValid() && Client)
 	{
@@ -256,6 +259,7 @@ bool FProfilerServiceManager::HandlePing( float DeltaTime )
 void FProfilerServiceManager::HandleServiceCaptureMessage( const FProfilerServiceCapture& Message, const TSharedRef<IMessageContext, ESPMode::ThreadSafe>& Context )
 {
 #if STATS
+	LLM_SCOPE_BYNAME(TEXT("SessionProfiler"));
 	const bool bRequestedCaptureState = Message.bRequestedCaptureState;
 	const bool bIsCapturing = FCommandStatsFile::Get().IsStatFileActive();
 
@@ -302,6 +306,8 @@ void FProfilerServiceManager::HandleServiceRequestMessage( const FProfilerServic
 	{
 		if( LastStatsFilename.IsEmpty() == false )
 		{
+			LLM_SCOPE_BYNAME(TEXT("SessionProfiler"));
+
 			FileTransferRunnable->EnqueueFileToSend( LastStatsFilename, Context->GetSender(), InstanceId );
 			LastStatsFilename.Empty();
 		}
@@ -335,6 +341,8 @@ void FProfilerServiceManager::HandleServiceSubscribeMessage( const FProfilerServ
 	const FMessageAddress& SenderAddress = Context->GetSender();
 	if( MessageEndpoint.IsValid() && Message.SessionId == SessionId && Message.InstanceId == InstanceId && !ClientData.Contains( SenderAddress ) )
 	{
+		LLM_SCOPE_BYNAME(TEXT("SessionProfiler"));
+
 		UE_LOG( LogProfilerService, Log, TEXT( "Subscribe Session: %s, Instance: %s" ), *SessionId.ToString(), *InstanceId.ToString() );
 
 		FClientData Data;
@@ -391,6 +399,7 @@ void FProfilerServiceManager::HandleNewFrame(int64 Frame)
 	// Called from the stats thread.
 #if STATS
 	DECLARE_SCOPE_CYCLE_COUNTER( TEXT( "FProfilerServiceManager::HandleNewFrame" ), STAT_FProfilerServiceManager_HandleNewFrame, STATGROUP_Profiler );
+	LLM_SCOPE_BYNAME(TEXT("SessionProfiler"));
 	
 	const FStatsThreadState& Stats = FStatsThreadState::GetLocalState();
 	const int32 CurrentMetadataSize = Stats.ShortNameToLongName.Num();
@@ -425,6 +434,7 @@ void FProfilerServiceManager::HandleNewFrame(int64 Frame)
 void FProfilerServiceManager::CompressDataAndSendToGame( TArray<uint8>* DataToTask, int64 Frame )
 {
 	DECLARE_SCOPE_CYCLE_COUNTER( TEXT( "FProfilerServiceManager::CompressDataAndSendToGame" ), STAT_FProfilerServiceManager_CompressDataAndSendToGame, STATGROUP_Profiler );
+	LLM_SCOPE_BYNAME(TEXT("SessionProfiler"));
 
 	const uint8* UncompressedPtr = DataToTask->GetData();
 	const int32 UncompressedSize = DataToTask->Num();
