Marked the Profiler* modules as deprecated. The old Profiler is actually deprecated since UE 5.0, superceeded by Trace/UnrealInsights.

diff --git a/Engine/Source/Developer/Profiler/Profiler.Build.cs b/Engine/Source/Developer/Profiler/Profiler.Build.cs
--- a/Engine/Source/Developer/Profiler/Profiler.Build.cs
+++ b/Engine/Source/Developer/Profiler/Profiler.Build.cs
@@ -1,10 +1,12 @@
 // Copyright Epic Games, Inc. All Rights Reserved.
 
+using System;
 using UnrealBuildTool;
 
+[Obsolete("Profiler is deprecated since UE 5.0 - use Trace/UnrealInsights instead.")]
 public class Profiler : ModuleRules
 {
-	public Profiler( ReadOnlyTargetRules Target ) : base(Target)
+	public Profiler(ReadOnlyTargetRules Target) : base(Target)
 	{
 		UnsafeTypeCastWarningLevel = WarningLevel.Error;
 
diff --git a/Engine/Source/Developer/ProfilerClient/ProfilerClient.Build.cs b/Engine/Source/Developer/ProfilerClient/ProfilerClient.Build.cs
--- a/Engine/Source/Developer/ProfilerClient/ProfilerClient.Build.cs
+++ b/Engine/Source/Developer/ProfilerClient/ProfilerClient.Build.cs
@@ -1,7 +1,10 @@
 // Copyright Epic Games, Inc. All Rights Reserved.
 
+using System;
+
 namespace UnrealBuildTool.Rules
 {
+	[Obsolete("Profiler is deprecated since UE 5.0 - use Trace/UnrealInsights instead.")]
 	public class ProfilerClient : ModuleRules
 	{
 		public ProfilerClient(ReadOnlyTargetRules Target) : base(Target)
diff --git a/Engine/Source/Developer/ProfilerMessages/ProfilerMessages.Build.cs b/Engine/Source/Developer/ProfilerMessages/ProfilerMessages.Build.cs
--- a/Engine/Source/Developer/ProfilerMessages/ProfilerMessages.Build.cs
+++ b/Engine/Source/Developer/ProfilerMessages/ProfilerMessages.Build.cs
@@ -1,20 +1,21 @@
 // Copyright Epic Games, Inc. All Rights Reserved.
 
+using System;
+
 namespace UnrealBuildTool.Rules
 {
+	[Obsolete("Profiler is deprecated since UE 5.0 - use Trace/UnrealInsights instead.")]
 	public class ProfilerMessages : ModuleRules
 	{
 		public ProfilerMessages(ReadOnlyTargetRules Target) : base(Target)
 		{
 			PublicDependencyModuleNames.AddRange(
-				new string[]
-				{
+				new string[] {
 					"CoreUObject",
 				});
 
 			PrivateDependencyModuleNames.AddRange(
-				new string[]
-				{
+				new string[] {
 					"Core",
 				});
 
diff --git a/Engine/Source/Developer/ProfilerService/ProfilerService.Build.cs b/Engine/Source/Developer/ProfilerService/ProfilerService.Build.cs
--- a/Engine/Source/Developer/ProfilerService/ProfilerService.Build.cs
+++ b/Engine/Source/Developer/ProfilerService/ProfilerService.Build.cs
@@ -1,21 +1,22 @@
 // Copyright Epic Games, Inc. All Rights Reserved.
 
+using System;
+
 namespace UnrealBuildTool.Rules
 {
+	[Obsolete("Profiler is deprecated since UE 5.0 - use Trace/UnrealInsights instead.")]
 	public class ProfilerService : ModuleRules
 	{
 		public ProfilerService(ReadOnlyTargetRules Target) : base(Target)
 		{
 			PublicDependencyModuleNames.AddRange(
-				new string[]
-				{
+				new string[] {
 					"Core",
 				}
 			);
 
 			PrivateDependencyModuleNames.AddRange(
-				new string[]
-				{
+				new string[] {
 					"CoreUObject",
 					"Engine",
 					"ProfilerMessages",
