Fix invliad log expression.

diff --git a/Engine/Source/Developer/ProfilerService/Private/ProfilerServiceManager.cpp b/Engine/Source/Developer/ProfilerService/Private/ProfilerServiceManager.cpp
--- a/Engine/Source/Developer/ProfilerService/Private/ProfilerServiceManager.cpp
+++ b/Engine/Source/Developer/ProfilerService/Private/ProfilerServiceManager.cpp
@@ -455,7 +455,7 @@ void FProfilerServiceManager::CompressDataAndSendToGame( TArray<uint8>* DataToTa
 	FProfilerServiceData2* ToGameThread = new FProfilerServiceData2( InstanceId, Frame, HexData, CompressedSize, UncompressedSize );
 
 	const float CompressionRatio = (float)UncompressedSize / (float)CompressedSize;
-	UE_LOG( LogProfilerService, VeryVerbose, TEXT( "Frame: %i, UncompressedSize: %i/%f, InstanceId: %i" ), ToGameThread->Frame, UncompressedSize, CompressionRatio, *InstanceId.ToString() );
+	UE_LOG( LogProfilerService, VeryVerbose, TEXT( "Frame: %i, UncompressedSize: %i/%f, InstanceId: %s" ), ToGameThread->Frame, UncompressedSize, CompressionRatio, *InstanceId.ToString() );
 
 	// Send to the game thread. PreviewClients is not thread-safe, so we cannot send the data here.
 	FSimpleDelegateGraphTask::CreateAndDispatchWhenReady
